AccumulationDOF: Experimental plugin that renders depth of field by accumulating deferred renders from multiple aperture view points, intended for production rendering (not gameplay).

Provides a way to render depth of field that greatly mitigates artifacts due to missing parallax, while preserving deferred rendering lookdev.
Use when not getting the desired visual results with the default DOF. Rendering time increases linearly with number of samples.

Main settings:

* Number of samples. The scene will be rendered from multiple aperture locations and accumulate the results.
* Splats size. To improve bokeh fill, it uses standard DOF as splats. The size is the fraction of the aperture diameter (defaults to 1/8), and kept small as to not reintroduce artifacts or over-soften the bokeh edges.

Usage:

* In-editor preview, under level viewport >> scalability dropdown >> Accumulation DOF
* MRQ integration via Accumulation DOF pass. Set spatial samples to 1, and temporal samples to the value you desire for high quality motion blur. (Note: Aperture renders are not currently amortized across temporal samples).

Feature support:

* Petzval swirly bokeh
* Cat's eye (barrel)
* Anamorphic squeeze
* Multi-band lateral chromatic aberration
* Multi-band axial chromatic aberration
* Spherical aberration
* User-customizable bokeh texture.

Main pending work:

* MRG integration
* Improve AA
* Improve motion blur

[RN] Accumulation DOF option for production rendering. Includes in-editor preview option.

#Synced-CL 49977182

diff --git a/Engine/Plugins/VirtualProduction/AccumulationDOF/Source/AccumulationDOF/Private/MoviePipelineAccumulationDOFPass.cpp b/Engine/Plugins/VirtualProduction/AccumulationDOF/Source/AccumulationDOF/Private/MoviePipelineAccumulationDOFPass.cpp
new file mode 100644
--- /dev/null
+++ b/Engine/Plugins/VirtualProduction/AccumulationDOF/Source/AccumulationDOF/Private/MoviePipelineAccumulationDOFPass.cpp
@@ -0,0 +1,1045 @@
+// Copyright Epic Games, Inc. All Rights Reserved.
+
+#include "MoviePipelineAccumulationDOFPass.h"
+
+#include "AccumulationDOFComponent.h"
+#include "AccumulationDOFMath.h"
+#include "AccumulationDOFSceneViewExtension.h"
+#include "AccumulationDOFShaders.h"
+#include "AccumulationDOFUtils.h"
+#include "ApertureSampler.h"
+
+#include "Camera/PlayerCameraManager.h"
+#include "CineCameraActor.h"
+#include "CineCameraComponent.h"
+#include "Engine/TextureRenderTarget.h"
+#include "Engine/TextureRenderTarget2D.h"
+#include "Engine/World.h"
+#include "EngineModule.h"
+#include "GameFramework/PlayerController.h"
+#include "HAL/FileManager.h"
+#include "HAL/IConsoleManager.h"
+#include "ImagePixelData.h"
+#include "Kismet/KismetRenderingLibrary.h"
+#include "Misc/Paths.h"
+#include "MoviePipeline.h"
+#include "MoviePipelineCameraSetting.h"
+#include "MoviePipelineHighResSetting.h"
+#include "MoviePipelineOutputBase.h"
+#include "MoviePipelineOutputBuilder.h"
+#include "MoviePipelinePrimaryConfig.h"
+#include "MovieRenderOverlappedImage.h"
+#include "MovieRenderPipelineDataTypes.h"
+#include "RenderingThread.h"
+#include "SceneManagement.h"
+#include "SceneView.h"
+#include "TextureResource.h"
+
+
+DEFINE_LOG_CATEGORY(LogAccumulationDOFMRQ);
+
+
+namespace AccumulationDOF
+{
+
+struct FDeferredPassRenderStatePayload : public UMoviePipelineImagePassBase::IViewCalcPayload
+{
+	int32 CameraIndex;
+	FIntPoint TileIndex;
+	int32 SceneViewIndex;
+};
+}
+
+UMoviePipelineAccumulationDOFPass::UMoviePipelineAccumulationDOFPass()
+	: UMoviePipelineDeferredPassBase()
+{
+	PassIdentifier = FMoviePipelinePassIdentifier("AccumulationDOF");
+}
+
+UMoviePipelineAccumulationDOFPass::~UMoviePipelineAccumulationDOFPass()
+{
+	if (ApertureSampler)
+	{
+		ApertureSampler->Shutdown();
+		ApertureSampler = nullptr;
+	}
+}
+
+void UMoviePipelineAccumulationDOFPass::SetupImpl(const MoviePipeline::FMoviePipelineRenderPassInitSettings& InPassInitSettings)
+{
+	Super::SetupImpl(InPassInitSettings);
+
+	// Allocate dedicated view state for injection renders.
+	InjectionViewState.Allocate(InPassInitSettings.FeatureLevel);
+
+	// Add Lumen scene data initialization
+	if (InjectionViewState.GetReference() && GetWorld() && GetWorld()->Scene)
+	{
+		InjectionViewState.GetReference()->AddLumenSceneData(GetWorld()->Scene, 1.0f);
+	}
+
+	// Override LUT size to minimize banding
+	SavedLUTSize = AccumulationDOFUtils::OverrideLUTSizeIfNeeded(bForceColorGradingLUTSize64);
+
+	InitializeSampler();
+}
+
+void UMoviePipelineAccumulationDOFPass::TeardownImpl()
+{
+	if (ApertureSampler)
+	{
+		ApertureSampler->Shutdown();
+		ApertureSampler = nullptr;
+	}
+
+	// Restore LUT size if it was overridden
+	AccumulationDOFUtils::RestoreLUTSize(SavedLUTSize);
+	SavedLUTSize = 0;
+
+	CachedCineCameraComponent.Reset();
+	CachedDOFComponent.Reset();
+	CaptureExtension.Reset();
+	InjectionExtension.Reset();
+	InjectionViewState.Destroy();
+
+	Super::TeardownImpl();
+}
+
+void UMoviePipelineAccumulationDOFPass::RenderSample_GameThreadImpl(const FMoviePipelineRenderPassMetrics& InSampleState)
+{
+	UMoviePipelineImagePassBase::RenderSample_GameThreadImpl(InSampleState);
+
+	// Cache motion blur fraction for this frame.
+	CachedMotionBlurFraction = InSampleState.OutputState.TimeData.MotionBlurFraction;
+
+	// Determine if spatiotemporal mode is active
+	const bool bSpatiotemporalMode = bSpatiotemporalSampling && InSampleState.OutputState.TemporalSampleCount > 1;
+
+	// In spatiotemporal mode, only update camera state on first temporal sample to avoid
+	// resetting accumulation mid-frame when camera parameters (like focus) change
+	const bool bShouldUpdateCamera = !bSpatiotemporalMode || InSampleState.OutputState.IsFirstTemporalSample();
+	if (bShouldUpdateCamera)
+	{
+		UpdateCameraDependentState();
+	}
+
+	if (!ApertureSampler || !ApertureSampler->IsInitialized())
+	{
+		UE_LOG(LogAccumulationDOFMRQ, Error, TEXT("ApertureSampler not initialized"));
+		return;
+	}
+
+	// During warmup frames, render a single sample to prime temporal history
+	if (InSampleState.bDiscardResult)
+	{
+		const int32 CameraIndex = 0;
+		AccumulationDOF::FDeferredPassRenderStatePayload Payload;
+		Payload.CameraIndex = CameraIndex;
+		Payload.TileIndex = InSampleState.TileIndexes;
+		Payload.SceneViewIndex = 0;
+
+		// Render center sample for warmup
+
+		CurrentApertureSampleIndex = bSpatiotemporalMode ?
+			ApertureSampler->GetCenterSampleIndex() : 0;
+
+		FMoviePipelineRenderPassMetrics SampleMetrics = GetRenderPassMetricsForCamera(CameraIndex, InSampleState);
+		SampleMetrics.OutputState.CameraIndex = CameraIndex;
+
+		TSharedPtr<FSceneViewFamilyContext> ViewFamily = CalculateViewFamily(SampleMetrics, &Payload);
+
+		if (ViewFamily.IsValid())
+		{
+			TWeakObjectPtr<UTextureRenderTarget2D> ViewRenderTarget = GetOrCreateViewRenderTarget(SampleMetrics.BackbufferSize, &Payload);
+			if (ViewRenderTarget.IsValid())
+			{
+				FRenderTarget* RenderTarget = ViewRenderTarget->GameThread_GetRenderTargetResource();
+
+				if (RenderTarget && ViewFamily->Views.Num() > 0)
+				{
+					FSceneView* View = const_cast<FSceneView*>(ViewFamily->Views[0]);
+					FCanvas Canvas = FCanvas(RenderTarget, nullptr, GetPipeline()->GetWorld(), View->GetFeatureLevel(), FCanvas::CDM_DeferDrawing, 1.0f);
+					GetRendererModule().BeginRenderingViewFamily(&Canvas, ViewFamily.Get());
+				}
+			}
+		}
+
+		CurrentApertureSampleIndex = -1;
+
+		return;
+	}
+
+	if (bSpatiotemporalMode)
+	{
+		RenderSamples_Spatiotemporal(InSampleState);
+	}
+	else
+	{
+		RenderSamples(InSampleState);
+	}
+}
+
+void UMoviePipelineAccumulationDOFPass::RenderSamples_Spatiotemporal(const FMoviePipelineRenderPassMetrics& InSampleState)
+{
+	const int32 TemporalIndex    = InSampleState.OutputState.TemporalSampleIndex;
+	const int32 TemporalCount    = InSampleState.OutputState.TemporalSampleCount;
+	const bool bIsFirstTemporal  = InSampleState.OutputState.IsFirstTemporalSample();
+	const bool bIsLastTemporal   = InSampleState.OutputState.IsLastTemporalSample();
+	const int32 ActualNumSamples = ApertureSampler->GetActualNumSamples();
+
+	if (ActualNumSamples <= 0 || TemporalCount <= 0)
+	{
+		return;
+	}
+
+	// Detect new frame
+	const int32 CurrentFrame = InSampleState.OutputState.OutputFrameNumber;
+	if (CurrentFrame != LastOutputFrameNumber)
+	{
+		LastOutputFrameNumber = CurrentFrame;
+	}
+
+	// Reset accumulation on first temporal sample of frame
+	if (bIsFirstTemporal)
+	{
+		ApertureSampler->ResetAccumulation();
+	}
+
+	// Determine aperture sample distribution across temporal samples
+
+	int32 StartApertureIndex, EndApertureIndex;
+
+	if (TemporalCount < ActualNumSamples)
+	{
+		// Spread all aperture samples across temporal samples evenly
+
+		const int32 BaseCount = ActualNumSamples / TemporalCount;
+		const int32 Remainder = ActualNumSamples % TemporalCount;
+
+		const int32 SamplesThisTemporal = (TemporalIndex < Remainder) ? (BaseCount + 1) : BaseCount;
+
+		if (TemporalIndex < Remainder)
+		{
+			StartApertureIndex = TemporalIndex * (BaseCount + 1);
+		}
+		else
+		{
+			StartApertureIndex = Remainder * (BaseCount + 1) + (TemporalIndex - Remainder) * BaseCount;
+		}
+
+		EndApertureIndex = StartApertureIndex + SamplesThisTemporal;
+	}
+	else if (TemporalCount == ActualNumSamples)
+	{
+		// 1:1 mapping
+		StartApertureIndex = TemporalIndex;
+		EndApertureIndex = TemporalIndex + 1;
+	}
+	else
+	{
+		// Oversampling - cycle through aperture samples
+
+		const int32 ApertureSampleIndex = TemporalIndex % ActualNumSamples;
+
+		StartApertureIndex = ApertureSampleIndex;
+		EndApertureIndex = ApertureSampleIndex + 1;
+	}
+
+	// Setup common rendering state
+
+	const int32 CameraIndex = 0;
+
+	FMoviePipelinePassIdentifier PassIdentifierForCurrentCamera;
+
+	PassIdentifierForCurrentCamera.Name = PassIdentifier.Name;
+	PassIdentifierForCurrentCamera.CameraName = GetCameraName(CameraIndex);
+
+	AccumulationDOF::FDeferredPassRenderStatePayload Payload;
+
+	Payload.CameraIndex = CameraIndex;
+	Payload.TileIndex = InSampleState.TileIndexes;
+	Payload.SceneViewIndex = 0;
+
+	// Render all aperture samples assigned to this temporal sample
+	for (int32 ApertureSampleIndex = StartApertureIndex; ApertureSampleIndex < EndApertureIndex; ++ApertureSampleIndex)
+	{
+		CurrentApertureSampleIndex = ApertureSampleIndex;
+
+		RenderSingleApertureSample(ApertureSampleIndex, InSampleState, Payload);
+
+		CurrentApertureSampleIndex = -1;
+	}
+
+	// Finalize accumulation on last temporal sample and submit result
+	if (bIsLastTemporal)
+	{
+		ApertureSampler->FinalizeAccumulation();
+
+		// Render injection pass and submit our accumulated result to MRQ
+		RenderInjectionAndSubmit(InSampleState, PassIdentifierForCurrentCamera);
+	}
+
+	// Note: We only submit the final accumulation
+}
+
+void UMoviePipelineAccumulationDOFPass::RenderSamples(const FMoviePipelineRenderPassMetrics& InSampleState)
+{
+	const int32 ActualNumSamples = ApertureSampler->GetActualNumSamples();
+
+	// Reset accumulation for new frame
+	ApertureSampler->ResetAccumulation();
+
+	// Get camera info for pass identifier
+
+	const int32 CameraIndex = 0; // Primary camera only for now
+
+	FMoviePipelinePassIdentifier PassIdentifierForCurrentCamera;
+	PassIdentifierForCurrentCamera.Name = PassIdentifier.Name;
+	PassIdentifierForCurrentCamera.CameraName = GetCameraName(CameraIndex);
+
+	// Create payload for view family calculation
+	AccumulationDOF::FDeferredPassRenderStatePayload Payload;
+	Payload.CameraIndex = CameraIndex;
+	Payload.TileIndex = InSampleState.TileIndexes;
+	Payload.SceneViewIndex = 0;
+
+	// Render all aperture samples in a loop
+	for (int32 SampleIndex = 0; SampleIndex < ActualNumSamples; ++SampleIndex)
+	{
+		CurrentApertureSampleIndex = SampleIndex;
+		RenderSingleApertureSample(SampleIndex, InSampleState, Payload);
+	}
+
+	// Reset sample index and capture extension after all samples rendered
+	CurrentApertureSampleIndex = -1;
+	CaptureExtension.Reset();
+
+	// Finalize accumulation
+	ApertureSampler->FinalizeAccumulation();
+
+	// Render injection pass and submit to MRQ
+	RenderInjectionAndSubmit(InSampleState, PassIdentifierForCurrentCamera);
+}
+
+void UMoviePipelineAccumulationDOFPass::RenderSingleApertureSample(
+	int32 SampleIndex,
+	const FMoviePipelineRenderPassMetrics& InSampleState,
+	AccumulationDOF::FDeferredPassRenderStatePayload& Payload)
+{
+	const int32 CameraIndex = Payload.CameraIndex;
+
+	// Create SVE to capture scene color
+	UTextureRenderTarget2D* SampleRT = ApertureSampler->GetSampleRT();
+
+	if (SampleRT)
+	{
+		FAccumulationDOFSVESettings SVESettings = FAccumulationDOFSVESettings::FromCineCameraComponent(CachedCineCameraComponent.Get());
+
+		CaptureExtension = MakeShared<FAccumulationDOFSceneViewExtension, ESPMode::ThreadSafe>(
+			SampleRT,
+			EAccumulationDOFSVEMode::Capture,
+			SVESettings
+		);
+
+		// Apply DOF splats if size is > 0
+		const float EffectiveDOFSplatSize = GetEffectiveDOFSplatSize();
+
+		if (EffectiveDOFSplatSize > UE_KINDA_SMALL_NUMBER)
+		{
+			const AccumulationDOFUtils::FCameraParams& CamParams = ApertureSampler->GetConfig().NumSamples > 0 ?
+				AccumulationDOFUtils::ExtractCameraParams(CachedCineCameraComponent.Get()) :
+				AccumulationDOFUtils::FCameraParams();
+
+			const float DOFSplatsFStop = AccumulationDOFMath::ComputeDOFSplatsFStop(
+				CamParams.ApertureRadiusCm,
+				CamParams.FocalLengthMm,
+				EffectiveDOFSplatSize
+			);
+
+			const float SafeSqueezeFactor = FMath::Max(CamParams.SqueezeFactor, UE_KINDA_SMALL_NUMBER);
+			if (DOFSplatsFStop > UE_KINDA_SMALL_NUMBER)
+			{
+				const float SensorWidthMm = CamParams.SensorWidthMm / SafeSqueezeFactor;
+				const bool bForceNeutralBokeh = true;
+
+				CaptureExtension->SetDOFSplatsSettings(DOFSplatsFStop, CamParams.FocusDistanceCm, SensorWidthMm, CamParams.SqueezeFactor, bForceNeutralBokeh);
+			}
+		}
+	}
+	else
+	{
+		CaptureExtension.Reset();
+	}
+
+	// Create view family for this aperture sample
+	FMoviePipelineRenderPassMetrics SampleMetrics = GetRenderPassMetricsForCamera(CameraIndex, InSampleState);
+	SampleMetrics.OutputState.CameraIndex = CameraIndex;
+
+	bInCapturePass = true;
+	TSharedPtr<FSceneViewFamilyContext> ViewFamily = CalculateViewFamily(SampleMetrics, &Payload);
+	bInCapturePass = false;
+
+	if (!ViewFamily.IsValid())
+	{
+		UE_LOG(LogAccumulationDOFMRQ, Error, TEXT("Failed to create view family for sample %d"), SampleIndex);
+		CaptureExtension.Reset();
+		return;
+	}
+
+	TWeakObjectPtr<UTextureRenderTarget2D> ViewRenderTarget = GetOrCreateViewRenderTarget(SampleMetrics.BackbufferSize, &Payload);
+
+	if (!ViewRenderTarget.IsValid())
+	{
+		UE_LOG(LogAccumulationDOFMRQ, Error, TEXT("Failed to get view render target for sample %d"), SampleIndex);
+		CaptureExtension.Reset();
+		return;
+	}
+
+	FRenderTarget* RenderTarget = ViewRenderTarget->GameThread_GetRenderTargetResource();
+
+	if (!RenderTarget)
+	{
+		UE_LOG(LogAccumulationDOFMRQ, Error, TEXT("Failed to get render target resource for sample %d"), SampleIndex);
+		CaptureExtension.Reset();
+		return;
+	}
+
+	if (ViewFamily->Views.Num() == 0)
+	{
+		UE_LOG(LogAccumulationDOFMRQ, Error, TEXT("ViewFamily has no views for sample %d"), SampleIndex);
+		CaptureExtension.Reset();
+		return;
+	}
+
+	// Render the scene with off-axis projection
+	FSceneView* View = const_cast<FSceneView*>(ViewFamily->Views[0]);
+
+	const int32 ActualNumSamples = ApertureSampler->GetActualNumSamples();
+
+	// If MRQ set a frame index override, combine it with aperture sample index for uniqueness.
+	// Uses modulo to prevent int32 overflow.
+	if (View->OverrideFrameIndexValue.IsSet())
+	{
+		const int64 BaseFrameIndex = View->OverrideFrameIndexValue.GetValue();
+		const int64 CombinedIndex = BaseFrameIndex * static_cast<int64>(ActualNumSamples) + SampleIndex;
+		View->OverrideFrameIndexValue = static_cast<int32>(CombinedIndex & 0x7FFFFFFF);
+	}
+
+	// Control when temporal history update with bWorldIsPaused
+
+	const bool bIsFirstApertureSample = (SampleIndex == 0);
+	const bool bIsLastApertureSample = (SampleIndex == ActualNumSamples - 1);
+
+	ETemporalHistoryMode HistoryMode = ETemporalHistoryMode::LastSampleOnly;
+
+	if (bOverrideTemporalHistoryMode)
+	{
+		HistoryMode = TemporalHistoryMode;
+	}
+	else if (CachedDOFComponent.IsValid())
+	{
+		HistoryMode = CachedDOFComponent->TemporalHistoryMode;
+	}
+
+	switch (HistoryMode)
+	{
+	case ETemporalHistoryMode::AllSamplesUpdate:
+
+		ViewFamily->bWorldIsPaused = false;
+		break;
+
+	case ETemporalHistoryMode::FirstSampleOnly:
+
+		ViewFamily->bWorldIsPaused = !bIsFirstApertureSample;
+		break;
+
+	case ETemporalHistoryMode::LastSampleOnly:
+
+		ViewFamily->bWorldIsPaused = !bIsLastApertureSample;
+		break;
+
+	case ETemporalHistoryMode::NoSamplesUpdate:
+
+		ViewFamily->bWorldIsPaused = true;
+		break;
+
+	default:
+		checkNoEntry();
+		break;
+	}
+
+	FCanvas Canvas = FCanvas(RenderTarget, nullptr, GetPipeline()->GetWorld(), View->GetFeatureLevel(), FCanvas::CDM_DeferDrawing, 1.0f);
+
+	GetRendererModule().BeginRenderingViewFamily(&Canvas, ViewFamily.Get());
+
+	// Capture SceneFringeIntensity from first sample
+	if (SampleIndex == 0 && CaptureExtension.IsValid())
+	{
+		const float CapturedValue = CaptureExtension->GetCapturedSceneFringeIntensity();
+		if (FMath::Abs(CapturedValue - CachedSceneFringeIntensity) > UE_KINDA_SMALL_NUMBER)
+		{
+			CachedSceneFringeIntensity = CapturedValue;
+		}
+	}
+
+	// Save individual sample if debug option enabled
+
+	if (bSaveDebugEXR && SampleRT)
+	{
+		FlushRenderingCommands();
+
+		FString SamplesDir = FPaths::ConvertRelativePathToFull(FPaths::ProjectSavedDir() / TEXT("AccumulationDOF_MRQ_Debug/Samples"));
+		IFileManager::Get().MakeDirectory(*SamplesDir, true);
+		FString Filename = FString::Printf(TEXT("Sample_%04d.exr"), SampleIndex);
+		UKismetRenderingLibrary::ExportRenderTarget(this, SampleRT, SamplesDir, Filename);
+	}
+
+	// Accumulate
+	const FVector2f ApertureOffset = ApertureSampler->GetApertureOffset(SampleIndex);
+	ApertureSampler->AccumulateExternalSample(SampleIndex, ApertureOffset);
+
+	// Release memory periodically
+	if (((SampleIndex + 1) % AccumulationDOFUtils::RenderingMemoryFlushBatchSize) == 0)
+	{
+		AccumulationDOFUtils::ReleaseRenderingMemory();
+	}
+
+	CaptureExtension.Reset();
+}
+
+void UMoviePipelineAccumulationDOFPass::RenderInjectionAndSubmit(
+	const FMoviePipelineRenderPassMetrics& InSampleState,
+	const FMoviePipelinePassIdentifier& PassIdentifierForCurrentCamera)
+{
+	const int32 CameraIndex = 0;
+	const bool bSpatiotemporalMode = bSpatiotemporalSampling && InSampleState.OutputState.TemporalSampleCount > 1;
+
+	// Get our accumulated result for injection
+	UTextureRenderTarget2D* NormalizedRT = ApertureSampler->GetAccumulatedResult();
+
+	if (!NormalizedRT)
+	{
+		UE_LOG(LogAccumulationDOFMRQ, Error, TEXT("No normalized render target available for injection pass"));
+		return;
+	}
+
+	// Determine motion blur setting
+	const bool bMotionBlurEnabled = bMotionBlur && !FMath::IsNearlyZero(CachedMotionBlurFraction);
+
+	// Create injection extension
+	{
+		// Choose injection mode based on motion blur
+		const EAccumulationDOFSVEMode InjectionMode = bMotionBlurEnabled
+			? EAccumulationDOFSVEMode::InjectViaTemporalUpscaler
+			: EAccumulationDOFSVEMode::Inject;
+
+		// Create injection extension
+		InjectionExtension = MakeShared<FAccumulationDOFSceneViewExtension, ESPMode::ThreadSafe>(
+			NormalizedRT,
+			InjectionMode
+		);
+
+		// Only allow engine's chromatic aberration if we didn't apply our own spectral / multi-band version of it.
+
+		const bool bUseSpectralLateralCA =
+			CachedDOFComponent.IsValid() ?
+			CachedDOFComponent->bSpectralLateralChromaticAberration
+			: true;
+
+		const bool bAppliedOurSpectralCA = bUseSpectralLateralCA && (CachedSceneFringeIntensity > UE_KINDA_SMALL_NUMBER);
+		InjectionExtension->SetAllowSceneFringe(!bAppliedOurSpectralCA);
+	}
+
+	bInInjectionPass = true;
+
+	AccumulationDOF::FDeferredPassRenderStatePayload Payload;
+	Payload.CameraIndex = CameraIndex;
+	Payload.TileIndex = InSampleState.TileIndexes;
+	Payload.SceneViewIndex = 0;
+
+	FMoviePipelineRenderPassMetrics InjectionMetrics = GetRenderPassMetricsForCamera(CameraIndex, InSampleState);
+	InjectionMetrics.OutputState.CameraIndex = CameraIndex;
+
+	TSharedPtr<FSceneViewFamilyContext> ViewFamily = CalculateViewFamily(InjectionMetrics, &Payload);
+	bInInjectionPass = false;
+
+	if (!ViewFamily.IsValid())
+	{
+		InjectionExtension.Reset();
+		UE_LOG(LogAccumulationDOFMRQ, Error, TEXT("Failed to create view family for injection pass"));
+		return;
+	}
+
+	TWeakObjectPtr<UTextureRenderTarget2D> ViewRenderTarget = GetOrCreateViewRenderTarget(InjectionMetrics.BackbufferSize, &Payload);
+
+	if (!ViewRenderTarget.IsValid())
+	{
+		InjectionExtension.Reset();
+		UE_LOG(LogAccumulationDOFMRQ, Error, TEXT("Failed to get view render target for injection pass"));
+		return;
+	}
+
+	FRenderTarget* RenderTarget = ViewRenderTarget->GameThread_GetRenderTargetResource();
+
+	if (!RenderTarget)
+	{
+		InjectionExtension.Reset();
+		UE_LOG(LogAccumulationDOFMRQ, Error, TEXT("Failed to get render target resource for injection pass"));
+		return;
+	}
+
+	if (ViewFamily->Views.Num() == 0)
+	{
+		InjectionExtension.Reset();
+		UE_LOG(LogAccumulationDOFMRQ, Error, TEXT("ViewFamily has no views for injection pass"));
+		return;
+	}
+
+	FSceneView* View = const_cast<FSceneView*>(ViewFamily->Views[0]);
+
+	FCanvas Canvas = FCanvas(RenderTarget, nullptr, GetPipeline()->GetWorld(), View->GetFeatureLevel(), FCanvas::CDM_DeferDrawing, 1.0f);
+	GetRendererModule().BeginRenderingViewFamily(&Canvas, ViewFamily.Get());
+
+	InjectionExtension.Reset();
+
+	// Submit to MRQ accumulator
+
+	FMoviePipelineRenderPassMetrics SubmissionState = InSampleState;
+
+	if (bSpatiotemporalMode)
+	{
+		// We've accumulated internally across temporal samples, so we tell MRQ this is a complete frame
+		SubmissionState.OutputState.TemporalSampleIndex = 0;
+		SubmissionState.OutputState.TemporalSampleCount = 1;
+	}
+
+	PostRendererSubmission(SubmissionState, PassIdentifierForCurrentCamera, GetOutputFileSortingOrder(), Canvas);
+
+	// Save debug EXR if enabled
+	if (bSaveDebugEXR && InSampleState.OutputState.IsLastTemporalSample())
+	{
+		SaveDebugEXR(InSampleState.OutputState.OutputFrameNumber);
+	}
+}
+
+AccumulationDOF::FApertureSamplerConfig UMoviePipelineAccumulationDOFPass::BuildSamplerConfig() const
+{
+	AccumulationDOF::FApertureSamplerConfig Config;
+
+	UMoviePipeline* Pipeline = GetPipeline();
+	if (!Pipeline)
+	{
+		return Config;
+	}
+
+	Config.World = Pipeline->GetWorld();
+
+	// Calculate resolution, taking into account number of tiles
+	{
+		const TArray<UMoviePipelineExecutorShot*>& ActiveShotList = Pipeline->GetActiveShotList();
+		const int32 CurrentShotIndex = Pipeline->GetCurrentShotIndex();
+		if (!ActiveShotList.IsValidIndex(CurrentShotIndex))
+		{
+			return Config;
+		}
+		UMoviePipelineExecutorShot* CurrentShot = ActiveShotList[CurrentShotIndex];
+
+		FIntPoint FullResolution = GetOverscannedResolutionForCamera(0);
+		UMoviePipelineHighResSetting* HighResSettings = Pipeline->FindOrAddSettingForShot<UMoviePipelineHighResSetting>(CurrentShot);
+
+		if (HighResSettings && HighResSettings->TileCount > 1)
+		{
+			FIntPoint TileSize = FIntPoint(
+				FMath::CeilToInt(static_cast<float>(FullResolution.X) / static_cast<float>(HighResSettings->TileCount)),
+				FMath::CeilToInt(static_cast<float>(FullResolution.Y) / static_cast<float>(HighResSettings->TileCount))
+			);
+			Config.Resolution = HighResSettings->CalculatePaddedBackbufferSize(TileSize);
+		}
+		else
+		{
+			Config.Resolution = FullResolution;
+		}
+
+		const UMoviePipelineCameraSetting* CameraSettings = Pipeline->FindOrAddSettingForShot<UMoviePipelineCameraSetting>(CurrentShot);
+
+		if (CameraSettings && CameraSettings->bOverrideCameraOverscan)
+		{
+			Config.OverscanFraction = FMath::Clamp(CameraSettings->OverscanPercentage, 0.0f, 1.0f);
+		}
+		else if (Pipeline->HasCachedCameraOverscan(0))
+		{
+			Config.OverscanFraction = Pipeline->GetCachedCameraOverscan(0);
+		}
+	}
+
+	// Sampling settings
+	Config.NumSamples = GetEffectiveNumApertureSamples();
+	Config.Mode = AccumulationDOF::EApertureSamplerMode::OneShot;
+
+	// Get settings from component if available
+	if (CachedDOFComponent.IsValid())
+	{
+		Config.SamplingPattern                     = CachedDOFComponent->SamplingPattern;
+		Config.DOFSplatSize	                       = GetEffectiveDOFSplatSize();
+		Config.AxialChromaticAberrationIntensity   = CachedDOFComponent->AxialChromaticAberrationIntensity;
+		Config.AxialChromaticAberrationNumBands    = CachedDOFComponent->AxialChromaticAberrationNumBands;
+		Config.bSpectralLateralChromaticAberration = CachedDOFComponent->bSpectralLateralChromaticAberration;
+		Config.W040                                = CachedDOFComponent->W040;
+		Config.BokehTexture                        = CachedDOFComponent->BokehTexture;
+		Config.bEnableBokehTexture                 = CachedDOFComponent->bEnableBokehTexture;
+		Config.WeightChannel                       = CachedDOFComponent->WeightChannel;
+		Config.TintStrength                        = CachedDOFComponent->TintStrength;
+		Config.BokehEdgeSoftness                   = CachedDOFComponent->BokehEdgeSoftness;
+		Config.TemporalHistoryMode                 = CachedDOFComponent->TemporalHistoryMode;
+		Config.bUseJitterAA                        = CachedDOFComponent->bUseJitterAA;
+	}
+	else
+	{
+		Config.SamplingPattern = EApertureSamplingPattern::Hexaweb;
+		Config.DOFSplatSize = GetEffectiveDOFSplatSize();
+	}
+
+	if (bOverrideTemporalHistoryMode)
+	{
+		Config.TemporalHistoryMode = TemporalHistoryMode;
+	}
+
+	Config.bUnifyWithTemporalSamples = bSpatiotemporalSampling;
+	Config.ViewModeIndex = VMI_Lit;
+
+	return Config;
+}
+
+AccumulationDOF::FApertureSamplerCameraState UMoviePipelineAccumulationDOFPass::BuildCameraState() const
+{
+	AccumulationDOF::FApertureSamplerCameraState State;
+
+	State.CineCameraComponent = CachedCineCameraComponent;
+
+	if (CachedCineCameraComponent.IsValid())
+	{
+		if (AActor* Owner = CachedCineCameraComponent->GetOwner())
+		{
+			State.CameraLocation = Owner->GetActorLocation();
+			State.CameraRotation = Owner->GetActorRotation();
+		}
+
+		State.bInitialized = true;
+	}
+
+	State.SceneFringeIntensity = CachedSceneFringeIntensity;
+
+	return State;
+}
+
+void UMoviePipelineAccumulationDOFPass::InitializeSampler()
+{
+	// Find camera first
+	CachedCineCameraComponent = FindCineCameraComponent();
+	FindAccumulationDOFComponent();
+
+	// Build config and camera state
+	AccumulationDOF::FApertureSamplerConfig Config = BuildSamplerConfig();
+	AccumulationDOF::FApertureSamplerCameraState CameraState = BuildCameraState();
+
+	// Create and initialize sampler
+	ApertureSampler = NewObject<UApertureSampler>(this);
+
+	if (!ApertureSampler->Initialize(Config, CameraState))
+	{
+		UE_LOG(LogAccumulationDOFMRQ, Error, TEXT("Failed to initialize ApertureSampler"));
+
+		ApertureSampler->Shutdown();
+		ApertureSampler = nullptr;
+
+		return;
+	}
+}
+
+void UMoviePipelineAccumulationDOFPass::UpdateCameraDependentState()
+{
+	// Update cached camera component
+	if (!CachedCineCameraComponent.IsValid())
+	{
+		CachedCineCameraComponent = FindCineCameraComponent();
+		FindAccumulationDOFComponent();
+	}
+
+	// Update sampler's camera state
+	if (ApertureSampler && CachedCineCameraComponent.IsValid())
+	{
+		AccumulationDOF::FApertureSamplerCameraState CameraState = BuildCameraState();
+		ApertureSampler->UpdateCameraState(CameraState);
+		bCameraParamsInitialized = true;
+	}
+	else if (!bCameraParamsInitialized)
+	{
+		UE_LOG(LogAccumulationDOFMRQ, Warning, TEXT("No CineCameraActor found. Using default values."));
+		bCameraParamsInitialized = true;
+	}
+}
+
+void UMoviePipelineAccumulationDOFPass::GatherOutputPassesImpl(TArray<FMoviePipelinePassIdentifier>& ExpectedRenderPasses)
+{
+	Super::GatherOutputPassesImpl(ExpectedRenderPasses);
+}
+
+void UMoviePipelineAccumulationDOFPass::ValidateStateImpl()
+{
+	Super::ValidateStateImpl();
+
+	if (bOverrideNumApertureSamples && NumApertureSamples < 1)
+	{
+		UE_LOG(LogAccumulationDOFMRQ, Warning, TEXT("NumApertureSamples must be at least 1. Using 1."));
+		NumApertureSamples = 1;
+	}
+}
+
+#if WITH_EDITOR
+FText UMoviePipelineAccumulationDOFPass::GetDisplayText() const
+{
+	return NSLOCTEXT("MovieRenderPipeline", "AccumulationDOFPassSetting_DisplayName", "Accumulation Depth of Field");
+}
+
+FText UMoviePipelineAccumulationDOFPass::GetFooterText(UMoviePipelineExecutorJob* InJob) const
+{
+	return FText::Format(
+		NSLOCTEXT("MovieRenderPipeline", "AccumulationDOFPassSetting_FooterText_Override", "{0} aperture samples (override)"),
+		FText::AsNumber(GetEffectiveNumApertureSamples())
+	);
+}
+#endif
+
+void UMoviePipelineAccumulationDOFPass::GetViewShowFlags(FEngineShowFlags& OutShowFlag, EViewModeIndex& OutViewModeIndex) const
+{
+	OutShowFlag = FEngineShowFlags(EShowFlagInitMode::ESFIM_Game);
+	OutViewModeIndex = VMI_Lit;
+
+	const bool bDOFSplatsEnabled = GetEffectiveDOFSplatSize() > UE_KINDA_SMALL_NUMBER;
+	OutShowFlag.SetDepthOfField(bDOFSplatsEnabled);
+
+	const bool bMotionBlurEnabled = bMotionBlur && !FMath::IsNearlyZero(CachedMotionBlurFraction);
+	OutShowFlag.SetMotionBlur(bMotionBlurEnabled && bInInjectionPass);
+
+	// When not using jitter AA, enable TAA to help smoothen DOF splats.
+
+	const bool bUseJitterAA = CachedDOFComponent.IsValid() ? CachedDOFComponent->bUseJitterAA : true;
+
+	OutShowFlag.SetTemporalAA(!bUseJitterAA);
+	OutShowFlag.SetAntiAliasing(!bUseJitterAA);
+}
+
+bool UMoviePipelineAccumulationDOFPass::IsAutoExposureAllowed(const FMoviePipelineRenderPassMetrics& InSampleState) const
+{
+	return Super::IsAutoExposureAllowed(InSampleState);
+}
+
+void UMoviePipelineAccumulationDOFPass::BlendPostProcessSettings(FSceneView* InView, FMoviePipelineRenderPassMetrics& InOutSampleState, IViewCalcPayload* OptPayload)
+{
+	Super::BlendPostProcessSettings(InView, InOutSampleState, OptPayload);
+
+	if (bInInjectionPass)
+	{
+		const bool bUseSpectralLateralCA = CachedDOFComponent.IsValid() ?
+			CachedDOFComponent->bSpectralLateralChromaticAberration 
+			: true;
+
+		const bool bAppliedOurSpectralCA = bUseSpectralLateralCA && (CachedSceneFringeIntensity > UE_KINDA_SMALL_NUMBER);
+
+		if (bAppliedOurSpectralCA)
+		{
+			InView->FinalPostProcessSettings.SceneFringeIntensity = 0.0f;
+		}
+	}
+}
+
+UE::MoviePipeline::FImagePassCameraViewData UMoviePipelineAccumulationDOFPass::GetCameraInfo(
+	FMoviePipelineRenderPassMetrics& InOutSampleState,
+	IViewCalcPayload* OptPayload) const
+{
+	UE::MoviePipeline::FImagePassCameraViewData CameraData = Super::GetCameraInfo(InOutSampleState, OptPayload);
+
+	// Apply projection and camera position offset for current aperture sample
+	if (ApertureSampler && CurrentApertureSampleIndex >= 0 && CurrentApertureSampleIndex < ApertureSampler->GetActualNumSamples())
+	{
+		const FVector2f ApertureOffset = ApertureSampler->GetApertureOffset(CurrentApertureSampleIndex);
+
+		// Offset camera position in the lens plane
+
+		const FRotator CameraRotation = CameraData.ViewInfo.Rotation;
+		const FVector CameraRight = FRotationMatrix(CameraRotation).GetUnitAxis(EAxis::Y);
+		const FVector CameraUp = FRotationMatrix(CameraRotation).GetUnitAxis(EAxis::Z);
+
+		CameraData.ViewInfo.Location += CameraRight * ApertureOffset.X + CameraUp * ApertureOffset.Y;
+
+		FMatrix ProjectionMatrix;
+		if (ApertureSampler->ComputeProjectionMatrix(CurrentApertureSampleIndex, -1.0f, ProjectionMatrix))
+		{
+			CameraData.bUseCustomProjectionMatrix = true;
+			CameraData.CustomProjectionMatrix = ProjectionMatrix;
+		}
+
+		// Reset this transform override that MRQ set. It has the center position but we need to use the correct jittered position for reprojection.
+		CameraData.ViewInfo.PreviousViewTransform.Reset();
+	}
+
+	return CameraData;
+}
+
+void UMoviePipelineAccumulationDOFPass::AddViewExtensions(FSceneViewFamilyContext& InContext, FMoviePipelineRenderPassMetrics& InOutSampleState)
+{
+	Super::AddViewExtensions(InContext, InOutSampleState);
+
+	if (bInCapturePass && CaptureExtension.IsValid())
+	{
+		InContext.ViewExtensions.Add(CaptureExtension.ToSharedRef());
+	}
+
+	if (bInInjectionPass && InjectionExtension.IsValid())
+	{
+		InContext.ViewExtensions.Add(InjectionExtension.ToSharedRef());
+	}
+}
+
+FSceneViewStateInterface* UMoviePipelineAccumulationDOFPass::GetSceneViewStateInterface(IViewCalcPayload* OptPayload)
+{
+	// During injection pass, always use dedicated injection view state
+
+	if (bInInjectionPass && InjectionViewState.GetReference())
+	{
+		return InjectionViewState.GetReference();
+	}
+
+	FSceneViewStateInterface* BaseState = Super::GetSceneViewStateInterface(OptPayload);
+
+	return BaseState;
+}
+
+UCineCameraComponent* UMoviePipelineAccumulationDOFPass::FindCineCameraComponent()
+{
+	UWorld* World = GetPipeline() ? GetPipeline()->GetWorld() : nullptr;
+
+	if (!World)
+	{
+		return nullptr;
+	}
+
+	APlayerController* PlayerController = World->GetFirstPlayerController();
+
+	if (!PlayerController || !PlayerController->PlayerCameraManager)
+	{
+		return nullptr;
+	}
+
+	ACineCameraActor* CineCameraActor = Cast<ACineCameraActor>(PlayerController->PlayerCameraManager->GetViewTarget());
+
+	if (!CineCameraActor)
+	{
+		return nullptr;
+	}
+
+	return CineCameraActor->GetCineCameraComponent();
+}
+
+UAccumulationDOFComponent* UMoviePipelineAccumulationDOFPass::FindAccumulationDOFComponent()
+{
+	if (CachedDOFComponent.IsValid())
+	{
+		return CachedDOFComponent.Get();
+	}
+
+	UWorld* World = GetPipeline() ? GetPipeline()->GetWorld() : nullptr;
+
+	if (!World)
+	{
+		return nullptr;
+	}
+
+	APlayerController* PlayerController = World->GetFirstPlayerController();
+
+	if (!PlayerController || !PlayerController->PlayerCameraManager)
+	{
+		return nullptr;
+	}
+
+	ACineCameraActor* CineCameraActor = Cast<ACineCameraActor>(PlayerController->PlayerCameraManager->GetViewTarget());
+
+	if (!CineCameraActor)
+	{
+		return nullptr;
+	}
+
+	UAccumulationDOFComponent* DOFComponent = CineCameraActor->FindComponentByClass<UAccumulationDOFComponent>();
+
+	if (DOFComponent)
+	{
+		CachedDOFComponent = DOFComponent;
+	}
+
+	return DOFComponent;
+}
+
+int32 UMoviePipelineAccumulationDOFPass::GetEffectiveNumApertureSamples() const
+{
+	if (bOverrideNumApertureSamples)
+	{
+		return NumApertureSamples;
+	}
+
+	if (CachedDOFComponent.IsValid())
+	{
+		return CachedDOFComponent->NumSamples;
+	}
+
+	return NumApertureSamples;
+}
+
+float UMoviePipelineAccumulationDOFPass::GetEffectiveDOFSplatSize() const
+{
+	if (bOverrideDOFSplatSize)
+	{
+		return DOFSplatSize;
+	}
+
+	if (CachedDOFComponent.IsValid())
+	{
+		return CachedDOFComponent->DOFSplatSize;
+	}
+
+	return DOFSplatSize;
+}
+
+void UMoviePipelineAccumulationDOFPass::SaveDebugEXR(int32 FrameNumber)
+{
+	FlushRenderingCommands();
+
+	FString Directory = FPaths::ConvertRelativePathToFull(FPaths::ProjectSavedDir() / TEXT("AccumulationDOF_MRQ_Debug"));
+	IFileManager::Get().MakeDirectory(*Directory, true);
+
+	// Save pre-PP result
+
+	if (ApertureSampler)
+	{
+		UTextureRenderTarget2D* NormalizedRT = ApertureSampler->GetAccumulatedResult();
+		if (NormalizedRT)
+		{
+			FString FilenamePrePP = FString::Printf(TEXT("MRQ_Frame%04d_PrePP.exr"), FrameNumber);
+			UKismetRenderingLibrary::ExportRenderTarget(this, NormalizedRT, Directory, FilenamePrePP);
+
+			UE_LOG(LogAccumulationDOFMRQ, Verbose, TEXT("Saved %s/%s"), *Directory, *FilenamePrePP);
+		}
+	}
+
+	// Save post-PP result
+
+	FIntPoint BackbufferSize = ApertureSampler ? ApertureSampler->GetConfig().Resolution : FIntPoint(1920, 1080);
+	TWeakObjectPtr<UTextureRenderTarget2D> ViewRenderTarget = GetOrCreateViewRenderTarget(BackbufferSize, nullptr);
+
+	if (ViewRenderTarget.IsValid())
+	{
+		FString FilenamePostPP = FString::Printf(TEXT("MRQ_Frame%04d_PostPP.exr"), FrameNumber);
+		UKismetRenderingLibrary::ExportRenderTarget(this, ViewRenderTarget.Get(), Directory, FilenamePostPP);
+
+		UE_LOG(LogAccumulationDOFMRQ, Verbose, TEXT("Saved %s/%s"), *Directory, *FilenamePostPP);
+	}
+}
diff --git a/Engine/Plugins/VirtualProduction/AccumulationDOF/Source/AccumulationDOF/Private/Rendering/AccumulationDOFTemporalUpscaler.cpp b/Engine/Plugins/VirtualProduction/AccumulationDOF/Source/AccumulationDOF/Private/Rendering/AccumulationDOFTemporalUpscaler.cpp
new file mode 100644
--- /dev/null
+++ b/Engine/Plugins/VirtualProduction/AccumulationDOF/Source/AccumulationDOF/Private/Rendering/AccumulationDOFTemporalUpscaler.cpp
@@ -0,0 +1,113 @@
+// Copyright Epic Games, Inc. All Rights Reserved.
+
+#include "AccumulationDOFTemporalUpscaler.h"
+#include "AccumulationDOFLog.h"
+#include "AccumulationDOFShaders.h"
+#include "RenderGraphBuilder.h"
+#include "RenderGraphUtils.h"
+#include "SceneView.h"
+#include "Engine/TextureRenderTarget2D.h"
+#include "TextureResource.h"
+
+using namespace AccumulationDOF;
+
+FAccumulationDOFTemporalUpscaler::FAccumulationDOFTemporalUpscaler(UTextureRenderTarget2D* InAccumulatedRT, float InProgressBarFraction)
+	: AccumulatedRT(InAccumulatedRT)
+	, ProgressBarFraction(InProgressBarFraction)
+{
+	check(InAccumulatedRT);
+}
+
+UE::Renderer::Private::ITemporalUpscaler::FOutputs FAccumulationDOFTemporalUpscaler::AddPasses(
+	FRDGBuilder& GraphBuilder,
+	const FSceneView& View,
+	const FInputs& Inputs) const
+{
+	UTextureRenderTarget2D* AccumulatedRTObj = AccumulatedRT.Get();
+	if (!ensure(AccumulatedRTObj))
+	{
+		FOutputs Outputs;
+		Outputs.FullRes = Inputs.SceneColor;
+		Outputs.NewHistory = MakeRefCount<FHistory>();
+		return Outputs;
+	}
+
+	FTextureRenderTargetResource* AccumulatedResource = AccumulatedRTObj->GameThread_GetRenderTargetResource();
+	if (!AccumulatedResource)
+	{
+		FOutputs Outputs;
+		Outputs.FullRes = Inputs.SceneColor;
+		Outputs.NewHistory = MakeRefCount<FHistory>();
+		return Outputs;
+	}
+
+	const FIntRect OutputViewRect = Inputs.OutputViewRect;
+
+	// Guard against invalid view rect (minimized viewport, disabled pass, etc.)
+	if (OutputViewRect.Width() <= 0 || OutputViewRect.Height() <= 0)
+	{
+		FOutputs Outputs;
+		Outputs.FullRes = Inputs.SceneColor;
+		Outputs.NewHistory = MakeRefCount<FHistory>();
+		return Outputs;
+	}
+
+	if (!Inputs.SceneColor.IsValid() || Inputs.SceneColor.Texture == nullptr)
+	{
+		FOutputs Outputs;
+		Outputs.FullRes = Inputs.SceneColor;
+		Outputs.NewHistory = MakeRefCount<FHistory>();
+		return Outputs;
+	}
+
+	FRHITexture* AccumulatedRHI = AccumulatedResource->GetRenderTargetTexture();
+	if (!AccumulatedRHI)
+	{
+		FOutputs Outputs;
+		Outputs.FullRes = Inputs.SceneColor;
+		Outputs.NewHistory = MakeRefCount<FHistory>();
+		return Outputs;
+	}
+
+	FRDGTextureRef AccumulatedRDG = GraphBuilder.RegisterExternalTexture(
+		CreateRenderTarget(AccumulatedRHI, TEXT("AccumulationDOFAccumulated"))
+	);
+
+	const FIntPoint OutputExtent = Inputs.SceneColor.Texture->Desc.Extent;
+	const EPixelFormat OutputFormat = Inputs.SceneColor.Texture->Desc.Format;
+
+	FRDGTextureDesc OutputDesc = FRDGTextureDesc::Create2D(
+		OutputExtent,
+		OutputFormat,
+		FClearValueBinding::Black,
+		TexCreate_RenderTargetable | TexCreate_ShaderResource
+	);
+
+	FRDGTextureRef OutputTexture = GraphBuilder.CreateTexture(OutputDesc, TEXT("AccumulationDOFUpscalerOutput"));
+
+	// Use our existing InjectToSceneColor shader to copy the accumulated result
+
+	const bool bDrawProgressBar = ProgressBarFraction >= 0.0f;
+
+	InjectToSceneColor(
+		GraphBuilder,
+		AccumulatedRDG,
+		OutputTexture,
+		OutputViewRect,
+		ProgressBarFraction,
+		bDrawProgressBar
+	);
+
+	FOutputs Outputs;
+
+	Outputs.FullRes = FScreenPassTexture(OutputTexture, OutputViewRect);
+	Outputs.NewHistory = MakeRefCount<FHistory>();
+
+	return Outputs;
+}
+
+UE::Renderer::Private::ITemporalUpscaler* FAccumulationDOFTemporalUpscaler::Fork_GameThread(
+	const FSceneViewFamily& ViewFamily) const
+{
+	return new FAccumulationDOFTemporalUpscaler(AccumulatedRT.Get(), ProgressBarFraction);
+}
diff --git a/Engine/Plugins/VirtualProduction/AccumulationDOF/Source/AccumulationDOF/Public/MoviePipelineAccumulationDOFPass.h b/Engine/Plugins/VirtualProduction/AccumulationDOF/Source/AccumulationDOF/Public/MoviePipelineAccumulationDOFPass.h
new file mode 100644
--- /dev/null
+++ b/Engine/Plugins/VirtualProduction/AccumulationDOF/Source/AccumulationDOF/Public/MoviePipelineAccumulationDOFPass.h
@@ -0,0 +1,227 @@
+﻿// Copyright Epic Games, Inc. All Rights Reserved.
+
+#pragma once
+
+#include "CoreMinimal.h"
+#include "AccumulationDOFTypes.h"
+#include "MoviePipelineDeferredPasses.h"
+
+#include "MoviePipelineAccumulationDOFPass.generated.h"
+
+ACCUMULATIONDOF_API DECLARE_LOG_CATEGORY_EXTERN(LogAccumulationDOFMRQ, Log, All);
+
+class ACineCameraActor;
+class FAccumulationDOFSceneViewExtension;
+class UApertureSampler;
+class UCineCameraComponent;
+class UAccumulationDOFComponent;
+class UTextureRenderTarget2D;
+class UTexture2D;
+
+namespace AccumulationDOF
+{
+	struct FApertureSamplerConfig;
+	struct FApertureSamplerCameraState;
+	struct FDeferredPassRenderStatePayload;
+}
+
+/**
+ * Movie Pipeline render pass for aperture-sampled depth of field.
+ *
+ * This pass renders multiple aperture samples with off-axis projection matrices and accumulates them 
+ * using GPU shaders to produce plausible DOF.
+ *
+ * Camera parameters (focal length, aperture, focus distance) are read from the CineCameraActor being 
+ * rendered by the sequence.
+ */
+UCLASS(BlueprintType)
+class ACCUMULATIONDOF_API UMoviePipelineAccumulationDOFPass : public UMoviePipelineDeferredPassBase
+{
+	GENERATED_BODY()
+
+public:
+	UMoviePipelineAccumulationDOFPass();
+	virtual ~UMoviePipelineAccumulationDOFPass();
+
+	/**
+	 * Override the number of aperture samples from the AccumulationDOFComponent.
+	 * When disabled, uses NumSamples from the AccumulationDOFComponent attached to the camera.
+	 */
+	UPROPERTY(EditAnywhere, BlueprintReadWrite, Category = "Accumulation DOF", meta = (InlineEditConditionToggle))
+	bool bOverrideNumApertureSamples = false;
+
+	/**
+	 * Number of aperture samples.
+	 * Used when bOverrideNumApertureSamples is enabled.
+	 */
+	UPROPERTY(EditAnywhere, BlueprintReadWrite, Category = "Accumulation DOF", meta = (ClampMin = "1", ClampMax = "65536", UIMin = "1", UIMax = "4096", EditCondition = "bOverrideNumApertureSamples"))
+	int32 NumApertureSamples = 512;
+
+	/**
+	 * Override the DOF splat size from the AccumulationDOFComponent.
+	 * When disabled, uses DOFSplatSize from the AccumulationDOFComponent attached to the camera.
+	 */
+	UPROPERTY(EditAnywhere, BlueprintReadWrite, Category = "Accumulation DOF", meta = (InlineEditConditionToggle))
+	bool bOverrideDOFSplatSize = false;
+
+	/**
+	 * Fraction of main aperture diameter for DOF splat size (0 = disabled).
+	 * DOF splats add small engine DOF blur to fill gaps between aperture samples.
+	 */
+	UPROPERTY(EditAnywhere, BlueprintReadWrite, Category = "Accumulation DOF",
+		meta = (DisplayName = "DOF Splat Size", ClampMin = "0.0", ClampMax = "1.0", UIMin = "0.0", UIMax = "1.0", EditCondition = "bOverrideDOFSplatSize"))
+	float DOFSplatSize = 0.125f;
+
+	/** Force color grading LUT size to 64 during rendering to reduce banding. Not exposed - always enabled. */
+	bool bForceColorGradingLUTSize64 = true;
+
+	/** Override the temporal history mode. Not exposed - LastSampleOnly works best. */
+	bool bOverrideTemporalHistoryMode = false;
+
+	/** Controls when temporal history (TAA, Lumen) updates during aperture sampling. Not exposed - LastSampleOnly works best. */
+	ETemporalHistoryMode TemporalHistoryMode = ETemporalHistoryMode::LastSampleOnly;
+
+	/**
+	 * Enable motion blur on the accumulated DOF result.
+	 */
+	UPROPERTY(EditAnywhere, BlueprintReadWrite, Category = "Accumulation DOF", meta = (DisplayName = "Motion Blur"))
+	bool bMotionBlur = false;
+
+	/**
+	 * Distribute aperture samples across MRQ temporal samples.
+	 * Motion blur is achieved through temporal accumulation rather than post-process.
+	 */
+	UPROPERTY(EditAnywhere, BlueprintReadWrite, Category = "Accumulation DOF",
+		meta = (DisplayName = "Spatiotemporal Sampling"))
+	bool bSpatiotemporalSampling = false;
+
+	/**
+	 * Save debug EXR files (individual samples and final accumulation).
+	 * Files are saved to the project's Saved/AccumulationDOF_MRQ_Debug folder.
+	 */
+	UPROPERTY(EditAnywhere, BlueprintReadWrite, Category = "Accumulation DOF|Debug")
+	bool bSaveDebugEXR = false;
+
+protected:
+
+	//~ UMoviePipelineRenderPass Interface
+	virtual void SetupImpl(const MoviePipeline::FMoviePipelineRenderPassInitSettings& InPassInitSettings) override;
+	virtual void TeardownImpl() override;
+	virtual void RenderSample_GameThreadImpl(const FMoviePipelineRenderPassMetrics& InSampleState) override;
+	virtual void GatherOutputPassesImpl(TArray<FMoviePipelinePassIdentifier>& ExpectedRenderPasses) override;
+	virtual void ValidateStateImpl() override;
+
+#if WITH_EDITOR
+	virtual FText GetDisplayText() const override;
+	virtual FText GetFooterText(UMoviePipelineExecutorJob* InJob) const override;
+#endif
+
+	virtual bool IsAntiAliasingSupported() const override
+	{
+		return false;
+	}
+
+	//~ UMoviePipelineImagePassBase Interface
+	virtual void GetViewShowFlags(FEngineShowFlags& OutShowFlag, EViewModeIndex& OutViewModeIndex) const override;
+	virtual UE::MoviePipeline::FImagePassCameraViewData GetCameraInfo(FMoviePipelineRenderPassMetrics& InOutSampleState, IViewCalcPayload* OptPayload = nullptr) const override;
+	virtual void AddViewExtensions(FSceneViewFamilyContext& InContext, FMoviePipelineRenderPassMetrics& InOutSampleState) override;
+	virtual bool IsAutoExposureAllowed(const FMoviePipelineRenderPassMetrics& InSampleState) const override;
+	virtual void BlendPostProcessSettings(FSceneView* InView, FMoviePipelineRenderPassMetrics& InOutSampleState, IViewCalcPayload* OptPayload = nullptr) override;
+	virtual FSceneViewStateInterface* GetSceneViewStateInterface(IViewCalcPayload* OptPayload = nullptr) override;
+
+private:
+
+	/** Aperture sampler that handles sample rendering and accumulation */
+	UPROPERTY(Transient)
+	TObjectPtr<UApertureSampler> ApertureSampler;
+
+	/** Current sample being rendered */
+	mutable int32 CurrentApertureSampleIndex = -1;
+
+	/** Cached reference to CineCameraComponent for SVE settings. */
+	UPROPERTY(Transient)
+	TWeakObjectPtr<UCineCameraComponent> CachedCineCameraComponent;
+
+	/** Cached reference to AccumulationDOFComponent (found in CineCameraActor). */
+	UPROPERTY(Transient)
+	TWeakObjectPtr<UAccumulationDOFComponent> CachedDOFComponent;
+
+	/** Scene View Extension for capture passes (one per aperture sample). */
+	TSharedPtr<FAccumulationDOFSceneViewExtension, ESPMode::ThreadSafe> CaptureExtension;
+
+	/** Scene View Extension for injection pass. */
+	TSharedPtr<FAccumulationDOFSceneViewExtension, ESPMode::ThreadSafe> InjectionExtension;
+
+	/** Flag to indicate we're doing a capture pass (so AddViewExtensions knows to add the capture SVE). */
+	bool bInCapturePass = false;
+
+	/** Flag to indicate we're doing the injection pass (so AddViewExtensions knows to add the injection SVE). */
+	bool bInInjectionPass = false;
+
+	/** Track whether we've initialized camera-dependent state. */
+	bool bCameraParamsInitialized = false;
+
+	/** Cached SceneFringeIntensity from the first sample's SVE (blended with PP volumes). */
+	float CachedSceneFringeIntensity = 0.0f;
+
+	/**  Separate view state for injection renders */
+	FSceneViewStateReference InjectionViewState;
+
+	/** Cached motion blur fraction from current frame's sample state. */
+	mutable float CachedMotionBlurFraction = 0.0f;
+
+	/** Whether this is the first injection render */
+	bool bFirstInjection = true;
+
+	/** Track current output frame for spatiotemporal mode reset detection */
+	int32 LastOutputFrameNumber = -1;
+
+	/** Saved LUT size before override (0 if not overridden) */
+	int32 SavedLUTSize = 0;
+
+private:
+
+	/** Build sampler config from MRQ settings and component. */
+	AccumulationDOF::FApertureSamplerConfig BuildSamplerConfig() const;
+
+	/** Build camera state from current sequence camera. */
+	AccumulationDOF::FApertureSamplerCameraState BuildCameraState() const;
+
+	/** Initialize the aperture sampler. */
+	void InitializeSampler();
+
+	/** Update camera-dependent state (called each frame). */
+	void UpdateCameraDependentState();
+
+	/** Find CineCameraComponent from MRQ's camera binding. */
+	UCineCameraComponent* FindCineCameraComponent();
+
+	/** Find AccumulationDOFComponent on the camera actor. */
+	UAccumulationDOFComponent* FindAccumulationDOFComponent();
+
+	/** Get effective number of aperture samples (incl. override). */
+	int32 GetEffectiveNumApertureSamples() const;
+
+	/** Get effective DOF splat size (incl. override). */
+	float GetEffectiveDOFSplatSize() const;
+
+	/** Save debug EXR files. */
+	void SaveDebugEXR(int32 FrameNumber);
+
+	/** Render aperture samples distributed among temporal samples */
+	void RenderSamples_Spatiotemporal(const FMoviePipelineRenderPassMetrics& InSampleState);
+
+	/** Render all aperture samples in one temporal sample */
+	void RenderSamples(const FMoviePipelineRenderPassMetrics& InSampleState);
+
+	/** Render a single aperture sample */
+	void RenderSingleApertureSample(
+		int32 SampleIndex,
+		const FMoviePipelineRenderPassMetrics& InSampleState,
+		AccumulationDOF::FDeferredPassRenderStatePayload& Payload);
+
+	/** Injection pass and submit to MRQ */
+	void RenderInjectionAndSubmit(
+		const FMoviePipelineRenderPassMetrics& InSampleState,
+		const FMoviePipelinePassIdentifier& PassIdentifierForCurrentCamera);
+};
