AccumulationDOF: Fix PVS warning.

#rnx

#Synced-CL 50383198

diff --git a/Engine/Plugins/VirtualProduction/AccumulationDOF/Source/AccumulationDOFEditor/Private/AccumulationDOFViewportExtension.cpp b/Engine/Plugins/VirtualProduction/AccumulationDOF/Source/AccumulationDOFEditor/Private/AccumulationDOFViewportExtension.cpp
--- a/Engine/Plugins/VirtualProduction/AccumulationDOF/Source/AccumulationDOFEditor/Private/AccumulationDOFViewportExtension.cpp
+++ b/Engine/Plugins/VirtualProduction/AccumulationDOF/Source/AccumulationDOFEditor/Private/AccumulationDOFViewportExtension.cpp
@@ -62,15 +62,15 @@ static float GetCameraFilmbackAspectRatio(const UCineCameraComponent* CineCam)
 	}
 
 	// Desqueezed aspect ratio
-	float AspectRatio = (SensorWidth * SqueezeFactor) / SensorHeight;
+	float FilmbackAspectRatio = (SensorWidth * SqueezeFactor) / SensorHeight;
 
 	// Apply crop aspect ratio
 	if (CineCam->CropSettings.AspectRatio > 0.0f)
 	{
-		AspectRatio = CineCam->CropSettings.AspectRatio;
+		FilmbackAspectRatio = CineCam->CropSettings.AspectRatio;
 	}
 
-	return AspectRatio;
+	return FilmbackAspectRatio;
 }
 
 
