AccumulationDOF: Add AA splatting per jittered pixel sample location.

Fix occasional RT aliasing when resizing viewport.

Also matched MRQ deferred pass in MotionBlur showflag but still controlling it via post process settings.

#rnx

#Synced-CL 50383198

diff --git a/Engine/Plugins/VirtualProduction/AccumulationDOF/Shaders/Private/AccumulationDOFAccumulate.usf b/Engine/Plugins/VirtualProduction/AccumulationDOF/Shaders/Private/AccumulationDOFAccumulate.usf
--- a/Engine/Plugins/VirtualProduction/AccumulationDOF/Shaders/Private/AccumulationDOFAccumulate.usf
+++ b/Engine/Plugins/VirtualProduction/AccumulationDOF/Shaders/Private/AccumulationDOFAccumulate.usf
@@ -71,6 +71,10 @@ float DiaphragmBladeCenterOffset;     // Blade arc center offset
 // AA reconstruction weight
 float AntiAliasingWeight;
 
+// Subpixel splatting 0-1
+float2 SubpixelOffset;
+float2 TexelSize;
+
 // Bokeh shape constants (matching DiaphragmDOF.h EBokehShape)
 #define BOKEH_SHAPE_CIRCLE          0
 #define BOKEH_SHAPE_STRAIGHT_BLADES 1
@@ -370,11 +374,11 @@ void MainPS(
 	out float4 OutColor   : SV_Target0
 )
 {
-	// Sample previous accumulation buffer
 	float4 PrevAccum = PrevAccumTexture.Sample(PrevAccumSampler, UV);
 
-	// Sample current aperture sample
-	float4 Sample = SampleTexture.Sample(SampleSampler, UV);
+	// Bilinear subpixel splatting
+	float2 GatherOffset = (SubpixelOffset - 0.5) * TexelSize;
+	float4 Sample = SampleTexture.Sample(SampleSampler, UV + GatherOffset);
 
 	// Map aperture offset to normalized disk coordinates [-1, 1]
 	float2 NormDisk = ApertureOffsetCm / max(ApertureRadiusCm, 0.0001);
@@ -487,6 +491,6 @@ void MainPS(
 	// Alpha: accumulate weight scaled by AntiAliasingWeight and avg(SpectralWeight) for correct normalization.
 	float AvgSpectralWeight = (SpectralWeight.r + SpectralWeight.g + SpectralWeight.b) / 3.0;
 	float WeightToAccumulate = (SampleMultiplier.r + SampleMultiplier.g + SampleMultiplier.b) / 3.0 * AntiAliasingWeight * AvgSpectralWeight;
-	
+
 	OutColor.a = PrevAccum.a + WeightToAccumulate;
 }
diff --git a/Engine/Plugins/VirtualProduction/AccumulationDOF/Source/AccumulationDOF/Private/MoviePipelineAccumulationDOFPass.cpp b/Engine/Plugins/VirtualProduction/AccumulationDOF/Source/AccumulationDOF/Private/MoviePipelineAccumulationDOFPass.cpp
--- a/Engine/Plugins/VirtualProduction/AccumulationDOF/Source/AccumulationDOF/Private/MoviePipelineAccumulationDOFPass.cpp
+++ b/Engine/Plugins/VirtualProduction/AccumulationDOF/Source/AccumulationDOF/Private/MoviePipelineAccumulationDOFPass.cpp
@@ -810,6 +810,36 @@ FText UMoviePipelineAccumulationDOFPass::GetFooterText(UMoviePipelineExecutorJob
 }
 #endif
 
+TSharedPtr<FSceneViewFamilyContext> UMoviePipelineAccumulationDOFPass::CalculateViewFamily(
+	FMoviePipelineRenderPassMetrics& InOutSampleState,
+	IViewCalcPayload* OptPayload)
+{
+	// Set ProjectionMatrixJitterAmount. This will later be applied via HackAddTemporalAAProjectionJitter.
+	if (ApertureSampler && CurrentApertureSampleIndex >= 0 && !InOutSampleState.bCameraCut)
+	{
+		InOutSampleState.ProjectionMatrixJitterAmount = ApertureSampler->GetJitterForProjectionMatrix(
+			CurrentApertureSampleIndex,
+			InOutSampleState.BackbufferSize);
+	}
+	else
+	{
+		InOutSampleState.ProjectionMatrixJitterAmount = FVector2D::ZeroVector;
+	}
+
+	TSharedPtr<FSceneViewFamilyContext> ViewFamily = Super::CalculateViewFamily(InOutSampleState, OptPayload);
+
+	if (!ViewFamily.IsValid())
+	{
+		return ViewFamily;
+	}
+
+	// Actual AA processing is controlled via AAMethod in BlendPostProcessSettings.
+	ViewFamily->EngineShowFlags.SetTemporalAA(true);
+	ViewFamily->EngineShowFlags.SetAntiAliasing(true);
+
+	return ViewFamily;
+}
+
 void UMoviePipelineAccumulationDOFPass::GetViewShowFlags(FEngineShowFlags& OutShowFlag, EViewModeIndex& OutViewModeIndex) const
 {
 	OutShowFlag = FEngineShowFlags(EShowFlagInitMode::ESFIM_Game);
@@ -819,9 +849,9 @@ void UMoviePipelineAccumulationDOFPass::GetViewShowFlags(FEngineShowFlags& OutSh
 	OutShowFlag.SetDepthOfField(bDOFSplatsEnabled);
 
 	const bool bMotionBlurEnabled = bMotionBlur && !FMath::IsNearlyZero(CachedMotionBlurFraction);
-	OutShowFlag.SetMotionBlur(bMotionBlurEnabled && bInInjectionPass);
 
-	// When not using jitter AA, enable TAA to help smoothen DOF splats.
+	// Enable motion blur ShowFlag actual motion blur is disabled for capture samples via MotionBlurAmount in BlendPostProcessSettings.
+	OutShowFlag.SetMotionBlur(bMotionBlurEnabled);
 
 	const bool bUseJitterAA = CachedDOFComponent.IsValid() ? CachedDOFComponent->bUseJitterAA : true;
 
@@ -838,6 +868,16 @@ void UMoviePipelineAccumulationDOFPass::BlendPostProcessSettings(FSceneView* InV
 {
 	Super::BlendPostProcessSettings(InView, InOutSampleState, OptPayload);
 
+	const bool bUseJitterAA = CachedDOFComponent.IsValid() ? CachedDOFComponent->bUseJitterAA : true;
+	const bool bCanUseTAA = !bUseJitterAA && InView->State != nullptr;
+	InView->AntiAliasingMethod = bCanUseTAA ? AAM_TemporalAA : AAM_None;
+
+	if (!bInInjectionPass)
+	{
+		InView->FinalPostProcessSettings.MotionBlurAmount = 0.0f;
+		InView->FinalPostProcessSettings.bOverride_MotionBlurAmount = true;
+	}
+
 	if (bInInjectionPass)
 	{
 		const bool bUseSpectralLateralCA = CachedDOFComponent.IsValid() ?
@@ -865,7 +905,6 @@ UE::MoviePipeline::FImagePassCameraViewData UMoviePipelineAccumulationDOFPass::G
 		const FVector2f ApertureOffset = ApertureSampler->GetApertureOffset(CurrentApertureSampleIndex);
 
 		// Offset camera position in the lens plane
-
 		const FRotator CameraRotation = CameraData.ViewInfo.Rotation;
 		const FVector CameraRight = FRotationMatrix(CameraRotation).GetUnitAxis(EAxis::Y);
 		const FVector CameraUp = FRotationMatrix(CameraRotation).GetUnitAxis(EAxis::Z);
@@ -873,13 +912,15 @@ UE::MoviePipeline::FImagePassCameraViewData UMoviePipelineAccumulationDOFPass::G
 		CameraData.ViewInfo.Location += CameraRight * ApertureOffset.X + CameraUp * ApertureOffset.Y;
 
 		FMatrix ProjectionMatrix;
+
 		if (ApertureSampler->ComputeProjectionMatrix(CurrentApertureSampleIndex, -1.0f, ProjectionMatrix))
 		{
 			CameraData.bUseCustomProjectionMatrix = true;
 			CameraData.CustomProjectionMatrix = ProjectionMatrix;
 		}
 
-		// Reset this transform override that MRQ set. It has the center position but we need to use the correct jittered position for reprojection.
+		// Reset PreviousViewTransform that MRQ base class set to previous frame's center position.
+		// With aperture sampling, keeping the center position gives Lumen false camera movement causing what looks like reprojection drift.
 		CameraData.ViewInfo.PreviousViewTransform.Reset();
 	}
 
diff --git a/Engine/Plugins/VirtualProduction/AccumulationDOF/Source/AccumulationDOF/Private/Rendering/AccumulationDOFSceneViewExtension.cpp b/Engine/Plugins/VirtualProduction/AccumulationDOF/Source/AccumulationDOF/Private/Rendering/AccumulationDOFSceneViewExtension.cpp
--- a/Engine/Plugins/VirtualProduction/AccumulationDOF/Source/AccumulationDOF/Private/Rendering/AccumulationDOFSceneViewExtension.cpp
+++ b/Engine/Plugins/VirtualProduction/AccumulationDOF/Source/AccumulationDOF/Private/Rendering/AccumulationDOFSceneViewExtension.cpp
@@ -125,8 +125,14 @@ void FAccumulationDOFSceneViewExtension::BeginRenderViewFamily(FSceneViewFamily&
 
 	if (Mode == EAccumulationDOFSVEMode::InjectViaTemporalUpscaler && RenderTarget.IsValid())
 	{
+		TRefCountPtr<FRHITexture> RHITexture;
+		if (FTextureRenderTargetResource* Resource = RenderTarget->GameThread_GetRenderTargetResource())
+		{
+			RHITexture = Resource->GetRenderTargetTexture();
+		}
+
 		FAccumulationDOFTemporalUpscaler* Upscaler = new FAccumulationDOFTemporalUpscaler(
-			RenderTarget.Get(),
+			RHITexture,
 			bDrawProgressBar ? ProgressBarFraction : -1.0f
 		);
 
diff --git a/Engine/Plugins/VirtualProduction/AccumulationDOF/Source/AccumulationDOF/Private/Rendering/AccumulationDOFTemporalUpscaler.cpp b/Engine/Plugins/VirtualProduction/AccumulationDOF/Source/AccumulationDOF/Private/Rendering/AccumulationDOFTemporalUpscaler.cpp
--- a/Engine/Plugins/VirtualProduction/AccumulationDOF/Source/AccumulationDOF/Private/Rendering/AccumulationDOFTemporalUpscaler.cpp
+++ b/Engine/Plugins/VirtualProduction/AccumulationDOF/Source/AccumulationDOF/Private/Rendering/AccumulationDOFTemporalUpscaler.cpp
@@ -1,21 +1,17 @@
 // Copyright Epic Games, Inc. All Rights Reserved.
 
 #include "AccumulationDOFTemporalUpscaler.h"
-#include "AccumulationDOFLog.h"
 #include "AccumulationDOFShaders.h"
 #include "RenderGraphBuilder.h"
 #include "RenderGraphUtils.h"
 #include "SceneView.h"
-#include "Engine/TextureRenderTarget2D.h"
-#include "TextureResource.h"
 
 using namespace AccumulationDOF;
 
-FAccumulationDOFTemporalUpscaler::FAccumulationDOFTemporalUpscaler(UTextureRenderTarget2D* InAccumulatedRT, float InProgressBarFraction)
-	: AccumulatedRT(InAccumulatedRT)
+FAccumulationDOFTemporalUpscaler::FAccumulationDOFTemporalUpscaler(TRefCountPtr<FRHITexture> InAccumulatedRHITexture, float InProgressBarFraction)
+	: AccumulatedRHITexture(InAccumulatedRHITexture)
 	, ProgressBarFraction(InProgressBarFraction)
 {
-	check(InAccumulatedRT);
 }
 
 UE::Renderer::Private::ITemporalUpscaler::FOutputs FAccumulationDOFTemporalUpscaler::AddPasses(
@@ -23,45 +19,12 @@ UE::Renderer::Private::ITemporalUpscaler::FOutputs FAccumulationDOFTemporalUpsca
 	const FSceneView& View,
 	const FInputs& Inputs) const
 {
-	UTextureRenderTarget2D* AccumulatedRTObj = AccumulatedRT.Get();
-	if (!ensure(AccumulatedRTObj))
-	{
-		FOutputs Outputs;
-		Outputs.FullRes = Inputs.SceneColor;
-		Outputs.NewHistory = MakeRefCount<FHistory>();
-		return Outputs;
-	}
-
-	FTextureRenderTargetResource* AccumulatedResource = AccumulatedRTObj->GameThread_GetRenderTargetResource();
-	if (!AccumulatedResource)
-	{
-		FOutputs Outputs;
-		Outputs.FullRes = Inputs.SceneColor;
-		Outputs.NewHistory = MakeRefCount<FHistory>();
-		return Outputs;
-	}
-
 	const FIntRect OutputViewRect = Inputs.OutputViewRect;
 
-	// Guard against invalid view rect (minimized viewport, disabled pass, etc.)
-	if (OutputViewRect.Width() <= 0 || OutputViewRect.Height() <= 0)
-	{
-		FOutputs Outputs;
-		Outputs.FullRes = Inputs.SceneColor;
-		Outputs.NewHistory = MakeRefCount<FHistory>();
-		return Outputs;
-	}
-
-	if (!Inputs.SceneColor.IsValid() || Inputs.SceneColor.Texture == nullptr)
-	{
-		FOutputs Outputs;
-		Outputs.FullRes = Inputs.SceneColor;
-		Outputs.NewHistory = MakeRefCount<FHistory>();
-		return Outputs;
-	}
-
-	FRHITexture* AccumulatedRHI = AccumulatedResource->GetRenderTargetTexture();
-	if (!AccumulatedRHI)
+	// Guard against invalid inputs
+	if (!AccumulatedRHITexture ||
+		OutputViewRect.Width() <= 0 || OutputViewRect.Height() <= 0 ||
+		!Inputs.SceneColor.IsValid() || Inputs.SceneColor.Texture == nullptr)
 	{
 		FOutputs Outputs;
 		Outputs.FullRes = Inputs.SceneColor;
@@ -70,7 +33,7 @@ UE::Renderer::Private::ITemporalUpscaler::FOutputs FAccumulationDOFTemporalUpsca
 	}
 
 	FRDGTextureRef AccumulatedRDG = GraphBuilder.RegisterExternalTexture(
-		CreateRenderTarget(AccumulatedRHI, TEXT("AccumulationDOFAccumulated"))
+		CreateRenderTarget(AccumulatedRHITexture.GetReference(), TEXT("AccumulationDOFAccumulated"))
 	);
 
 	const FIntPoint OutputExtent = Inputs.SceneColor.Texture->Desc.Extent;
@@ -109,5 +72,5 @@ UE::Renderer::Private::ITemporalUpscaler::FOutputs FAccumulationDOFTemporalUpsca
 UE::Renderer::Private::ITemporalUpscaler* FAccumulationDOFTemporalUpscaler::Fork_GameThread(
 	const FSceneViewFamily& ViewFamily) const
 {
-	return new FAccumulationDOFTemporalUpscaler(AccumulatedRT.Get(), ProgressBarFraction);
+	return new FAccumulationDOFTemporalUpscaler(AccumulatedRHITexture, ProgressBarFraction);
 }
diff --git a/Engine/Plugins/VirtualProduction/AccumulationDOF/Source/AccumulationDOF/Private/Rendering/AccumulationDOFTemporalUpscaler.h b/Engine/Plugins/VirtualProduction/AccumulationDOF/Source/AccumulationDOF/Private/Rendering/AccumulationDOFTemporalUpscaler.h
--- a/Engine/Plugins/VirtualProduction/AccumulationDOF/Source/AccumulationDOF/Private/Rendering/AccumulationDOFTemporalUpscaler.h
+++ b/Engine/Plugins/VirtualProduction/AccumulationDOF/Source/AccumulationDOF/Private/Rendering/AccumulationDOFTemporalUpscaler.h
@@ -3,12 +3,10 @@
 #pragma once
 
 #include "CoreMinimal.h"
+#include "RHI.h"
+#include "ScreenPass.h"
 #include "Templates/RefCounting.h"
 #include "TemporalUpscaler.h"
-#include "ScreenPass.h"
-#include "UObject/WeakObjectPtrTemplates.h"
-
-class UTextureRenderTarget2D;
 
 /**
  * Custom temporal upscaler that injects the accumulated AccumulationDOF result
@@ -63,7 +61,7 @@ public:
 		return TEXT("AccumulationDOF");
 	}
 
-	FAccumulationDOFTemporalUpscaler(UTextureRenderTarget2D* InAccumulatedRT, float InProgressBarFraction = -1.0f);
+	FAccumulationDOFTemporalUpscaler(TRefCountPtr<FRHITexture> InAccumulatedRHITexture, float InProgressBarFraction = -1.0f);
 
 	virtual ~FAccumulationDOFTemporalUpscaler() = default;
 
@@ -95,8 +93,8 @@ public:
 
 private:
 
-	/** The render target containing the accumulated DOF result to inject */
-	TWeakObjectPtr<UTextureRenderTarget2D> AccumulatedRT;
+	/** Accumulated RT, for render thread access */
+	TRefCountPtr<FRHITexture> AccumulatedRHITexture;
 
 	/** Progress bar fraction (0-1), or negative to disable */
 	float ProgressBarFraction;
diff --git a/Engine/Plugins/VirtualProduction/AccumulationDOF/Source/AccumulationDOF/Private/Rendering/ApertureSampler.cpp b/Engine/Plugins/VirtualProduction/AccumulationDOF/Source/AccumulationDOF/Private/Rendering/ApertureSampler.cpp
--- a/Engine/Plugins/VirtualProduction/AccumulationDOF/Source/AccumulationDOF/Private/Rendering/ApertureSampler.cpp
+++ b/Engine/Plugins/VirtualProduction/AccumulationDOF/Source/AccumulationDOF/Private/Rendering/ApertureSampler.cpp
@@ -23,10 +23,10 @@
 
 static TAutoConsoleVariable<int32> CVarAccumulationDOFPostMedianFilterEnable(
 	TEXT("r.AccumulationDOF.PostMedianFilter.Enable"),
-	1,
+	0,
 	TEXT("Enable the 3x3 median filter in Accumulation DOF.\n")
 	TEXT("  0: Disabled\n")
-	TEXT("  1: Enabled (default)"),
+	TEXT("  1: Enabled"),
 	ECVF_Default
 );
 
@@ -34,7 +34,7 @@ static TAutoConsoleVariable<int32> CVarAccumulationDOFPreMedianFilterEnable(
 	TEXT("r.AccumulationDOF.PreMedianFilter.Enable"),
 	0,
 	TEXT("Apply 3x3 median filter to each sample before accumulation.\n")
-	TEXT("  0: Disabled (default)\n")
+	TEXT("  0: Disabled\n")
 	TEXT("  1: Enabled"),
 	ECVF_Default
 );
@@ -73,13 +73,31 @@ enum class EAAFilter : int32
 
 static TAutoConsoleVariable<int32> CVarAccumulationDOFJitterAAFilter(
 	TEXT("r.AccumulationDOF.JitterAA.Filter"),
-	static_cast<int32>(EAAFilter::Mitchell),
+	static_cast<int32>(EAAFilter::Box),
 	TEXT("Reconstruction filter for AA.\n")
 	TEXT("  0: Box\n")
 	TEXT("  1: Mitchell-Netravali"),
 	ECVF_Default
 );
 
+static TAutoConsoleVariable<int32> CVarAccumulationDOFJitterAABilinearSplat(
+	TEXT("r.AccumulationDOF.JitterAA.BilinearSplat"),
+	1,
+	TEXT("Enable bilinear subpixel splatting (via gather) for AA reconstruction \n")
+	TEXT("  0: Disabled (point sampling)\n")
+	TEXT("  1: Enabled (default)"),
+	ECVF_Default
+);
+
+static TAutoConsoleVariable<int32> CVarAccumulationDOFApertureJitterAllowed(
+	TEXT("r.AccumulationDOF.ApertureJitter.Allowed"),
+	1,
+	TEXT("Allow aperture position jittering.\n")
+	TEXT("  0: All samples render from aperture center\n")
+	TEXT("  1: Normal aperture sampling"),
+	ECVF_Default
+);
+
 using namespace AccumulationDOF;
 
 //
@@ -234,6 +252,11 @@ void UApertureSampler::Shutdown()
 	TotalSpectralWeightPerChannel = FVector3f::ZeroVector;
 	CachedSceneFringeIntensity = 0.0f;
 
+	if (SampleRT || AccumA || AccumB || PrefilterRT || PostProcessRT)
+	{
+		FlushRenderingCommands();
+	}
+
 	SampleRT = nullptr;
 	AccumA = nullptr;
 	AccumB = nullptr;
@@ -573,7 +596,8 @@ bool UApertureSampler::RenderAllSamples()
 				const FVector3f RawWeight = ComputeSpectralWeight(NormalizedSpectralPosition);
 				const FVector3f SpectralWeight = RawWeight * InvTotal;
 
-				RenderSample(SampleIdx, ApertureOffsetsCm[SampleIdx], FocusDistBand);
+				const FVector2f SampleApertureOffsetCm = GetApertureOffset(SampleIdx);
+				RenderSample(SampleIdx, SampleApertureOffsetCm, FocusDistBand);
 
 				// Save intermediate sample if requested
 				if (Config.bSaveIntermediateSamples)
@@ -581,7 +605,7 @@ bool UApertureSampler::RenderAllSamples()
 					SaveSampleImage(IterationCount, Config.OutputDirectory, Config.OutputFilenamePrefix);
 				}
 
-				AccumulateSampleInternal(SampleIdx, ApertureOffsetsCm[SampleIdx], SpectralWeight);
+				AccumulateSampleInternal(SampleIdx, SampleApertureOffsetCm, SpectralWeight);
 
 				++IterationCount;
 				Progress.CurrentIteration = IterationCount;
@@ -621,7 +645,8 @@ bool UApertureSampler::RenderAllSamples()
 				break;
 			}
 
-			RenderSample(SampleIdx, ApertureOffsetsCm[SampleIdx]);
+			const FVector2f SampleApertureOffsetCm = GetApertureOffset(SampleIdx);
+			RenderSample(SampleIdx, SampleApertureOffsetCm);
 
 			// Save intermediate sample if requested
 			if (Config.bSaveIntermediateSamples)
@@ -629,7 +654,7 @@ bool UApertureSampler::RenderAllSamples()
 				SaveSampleImage(SampleIdx, Config.OutputDirectory, Config.OutputFilenamePrefix);
 			}
 
-			AccumulateSampleInternal(SampleIdx, ApertureOffsetsCm[SampleIdx], FVector3f(1.0f));
+			AccumulateSampleInternal(SampleIdx, SampleApertureOffsetCm, FVector3f(1.0f));
 
 			Progress.CurrentIteration = SampleIdx + 1;
 			NotifyProgress();
@@ -726,7 +751,7 @@ bool UApertureSampler::RenderAmortizedSamples()
 			break;
 		}
 
-		const FVector2f& ApertureOffsetCm = ApertureOffsetsCm[SampleIndex];
+		const FVector2f ApertureOffsetCm = GetApertureOffset(SampleIndex);
 
 		FVector3f SpectralWeight(1.0f);
 		float FocusDistanceOverride = -1.0f;
@@ -804,12 +829,7 @@ FApertureSampleParams UApertureSampler::ComputeSampleParams(int32 SampleIndex, f
 		return Params;
 	}
 
-	// Get aperture offset
-	FVector2f ApertureOffsetCm = FVector2f::ZeroVector;
-	if (SampleIndex >= 0 && SampleIndex < ApertureOffsetsCm.Num())
-	{
-		ApertureOffsetCm = ApertureOffsetsCm[SampleIndex];
-	}
+	const FVector2f ApertureOffsetCm = GetApertureOffset(SampleIndex);
 
 	// Compute projection matrix
 	FMatrix ProjectionMatrix;
@@ -917,14 +937,7 @@ bool UApertureSampler::ComputeProjectionMatrix(
 		return false;
 	}
 
-	// Get aperture offset
-
-	FVector2f ApertureOffsetCm = FVector2f::ZeroVector;
-
-	if (SampleIndex >= 0 && SampleIndex < ApertureOffsetsCm.Num())
-	{
-		ApertureOffsetCm = ApertureOffsetsCm[SampleIndex];
-	}
+	const FVector2f ApertureOffsetCm = GetApertureOffset(SampleIndex);
 
 	// Camera vectors
 	const FQuat CameraQuat = CameraState.CameraRotation.Quaternion();
@@ -1040,29 +1053,7 @@ bool UApertureSampler::ComputeProjectionMatrix(
 		Top    = FMath::Max(Top, Y);
 	}
 
-	// Apply subpixel jitter for AA if enabled
-
-	const bool bJitterAAAllowed = !!CVarAccumulationDOFJitterAAAllow.GetValueOnAnyThread();
-
-	if (Config.bUseJitterAA && SampleRT && bJitterAAAllowed)
-	{
-		const FVector2f JitterNorm = GetJitterOffset(SampleIndex);
-
-		const float FrustumWidth = Right - Left;
-		const float FrustumHeight = Top - Bottom;
-		const float SafeSizeX = FMath::Max(1.0f, static_cast<float>(SampleRT->SizeX));
-		const float SafeSizeY = FMath::Max(1.0f, static_cast<float>(SampleRT->SizeY));
-		const float PixelWidth = FrustumWidth / SafeSizeX;
-		const float PixelHeight = FrustumHeight / SafeSizeY;
-
-		const float OffsetX = JitterNorm.X * PixelWidth;
-		const float OffsetY = JitterNorm.Y * PixelHeight;
-
-		Left   += OffsetX;
-		Right  += OffsetX;
-		Bottom += OffsetY;
-		Top    += OffsetY;
-	}
+	// Note: Jitter not applied here but is applied later using the sample state ProjectionMatrixJitterAmount.
 
 	OutMatrix = AccumulationDOFMath::BuildOffAxisProjectionMatrix(Left, Right, Bottom, Top, NearCm, FarCm);
 
@@ -1172,10 +1163,13 @@ void UApertureSampler::AccumulateSampleInternal(
 	AccumParams.DiaphragmBladeRadius = BokehShape.DiaphragmBladeRadius;
 	AccumParams.DiaphragmBladeCenterOffset = BokehShape.DiaphragmBladeCenterOffset;
 
-	// AA reconstruction weight
+	// AA reconstruction weight and bilinear splatting, but skip for single sample which allows direct comparison with MRQ deferred pass.
+
 	const bool bJitterAAAllowed = !!CVarAccumulationDOFJitterAAAllow.GetValueOnAnyThread();
+	const bool bMultipleSamples = Progress.ActualNumSamples > 1;
 	const EAAFilter AAFilter = static_cast<EAAFilter>(CVarAccumulationDOFJitterAAFilter.GetValueOnAnyThread());
-	if (Config.bUseJitterAA && bJitterAAAllowed)
+
+	if (Config.bUseJitterAA && bJitterAAAllowed && bMultipleSamples)
 	{
 		const FVector2f JitterOffset = GetJitterOffset(SampleIndex);
 		switch (AAFilter)
@@ -1188,12 +1182,27 @@ void UApertureSampler::AccumulateSampleInternal(
 			AccumParams.AntiAliasingWeight = 1.0f;
 			break;
 		}
+
+		// Bilinear subpixel splatting via gather.
+
+		const bool bBilinearSplatEnabled = !!CVarAccumulationDOFJitterAABilinearSplat.GetValueOnAnyThread();
+
+		AccumParams.SubpixelOffset = bBilinearSplatEnabled
+			? FVector2f(0.5f, 0.5f) + JitterOffset
+			: FVector2f(0.5f, 0.5f);
 	}
 	else
 	{
 		AccumParams.AntiAliasingWeight = 1.0f;
+		AccumParams.SubpixelOffset = FVector2f(0.5f, 0.5f);
 	}
 
+	// Texel size for bilinear splatting UV offset
+	AccumParams.TexelSize = FVector2f(
+		1.0f / FMath::Max(1.0f, static_cast<float>(Config.Resolution.X)),
+		1.0f / FMath::Max(1.0f, static_cast<float>(Config.Resolution.Y))
+	);
+
 	// Execute accumulation
 	ENQUEUE_RENDER_COMMAND(AccumulationDOFAccumulate)(
 		[PrevAccumObj, SampleObj, OutputObj, BokehTextureRHI, BokehSamplerRHI, AccumParams]
@@ -1690,7 +1699,8 @@ void UApertureSampler::SaveFinalImage(const FString& Directory, const FString& P
 
 FVector2f UApertureSampler::GetApertureOffset(int32 SampleIndex) const
 {
-	if (SampleIndex >= 0 && SampleIndex < ApertureOffsetsCm.Num())
+	const bool bApertureJitterAllowed = !!CVarAccumulationDOFApertureJitterAllowed.GetValueOnAnyThread();
+	if (bApertureJitterAllowed && SampleIndex >= 0 && SampleIndex < ApertureOffsetsCm.Num())
 	{
 		return ApertureOffsetsCm[SampleIndex];
 	}
@@ -1712,6 +1722,30 @@ FVector2f UApertureSampler::GetJitterOffset(int32 SampleIndex) const
 	}
 }
 
+FVector2D UApertureSampler::GetJitterForProjectionMatrix(int32 SampleIndex, const FIntPoint& Resolution) const
+{
+	const bool bJitterAAAllowed = !!CVarAccumulationDOFJitterAAAllow.GetValueOnAnyThread();
+
+	if (!Config.bUseJitterAA || !bJitterAAAllowed || Resolution.X <= 0 || Resolution.Y <= 0)
+	{
+		return FVector2D::ZeroVector;
+	}
+
+	// No jitter for single sample
+	if (Progress.ActualNumSamples <= 1)
+	{
+		return FVector2D::ZeroVector;
+	}
+
+	const FVector2f JitterNorm = GetJitterOffset(SampleIndex);
+
+	// Convert to clip space.
+	return FVector2D(
+		JitterNorm.X * 2.0 / Resolution.X,
+		JitterNorm.Y * -2.0 / Resolution.Y
+	);
+}
+
 float UApertureSampler::ComputeSpectralFocusDistance(float NominalFocusCm, float NormalizedSpectralPosition) const
 {
 	// Blue (0) focuses closer, red (1) focuses farther
diff --git a/Engine/Plugins/VirtualProduction/AccumulationDOF/Source/AccumulationDOF/Private/Rendering/ApertureSampler.h b/Engine/Plugins/VirtualProduction/AccumulationDOF/Source/AccumulationDOF/Private/Rendering/ApertureSampler.h
--- a/Engine/Plugins/VirtualProduction/AccumulationDOF/Source/AccumulationDOF/Private/Rendering/ApertureSampler.h
+++ b/Engine/Plugins/VirtualProduction/AccumulationDOF/Source/AccumulationDOF/Private/Rendering/ApertureSampler.h
@@ -544,6 +544,17 @@ public:
 	 */
 	FVector2f GetJitterOffset(int32 SampleIndex) const;
 
+	/**
+	 * Get AA jitter in clip-space format for HackAddTemporalAAProjectionJitter().
+	 * This should be called after view creation and applied via View->ViewMatrices.HackAddTemporalAAProjectionJitter().
+	 *
+	 * @param SampleIndex - Index of the aperture sample
+	 * @param Resolution  - Render target resolution
+	 *
+	 * @return Jitter offset in clip-space units (matching MRQ convention with negated Y)
+	 */
+	FVector2D GetJitterForProjectionMatrix(int32 SampleIndex, const FIntPoint& Resolution) const;
+
 	/**
 	 * Get array ref. of aperture offsets.
 	 */
diff --git a/Engine/Plugins/VirtualProduction/AccumulationDOF/Source/AccumulationDOF/Private/Rendering/ApertureSamplingRenderer.cpp b/Engine/Plugins/VirtualProduction/AccumulationDOF/Source/AccumulationDOF/Private/Rendering/ApertureSamplingRenderer.cpp
--- a/Engine/Plugins/VirtualProduction/AccumulationDOF/Source/AccumulationDOF/Private/Rendering/ApertureSamplingRenderer.cpp
+++ b/Engine/Plugins/VirtualProduction/AccumulationDOF/Source/AccumulationDOF/Private/Rendering/ApertureSamplingRenderer.cpp
@@ -4,6 +4,7 @@
 
 #include "AccumulationDOFLog.h"
 #include "AccumulationDOFSceneViewExtension.h"
+#include "ApertureSampler.h"
 
 #include "Camera/CameraComponent.h"
 #include "CanvasTypes.h"
@@ -526,6 +527,17 @@ FSceneView* UApertureSamplingRenderer::CreateSceneView(
 
 	View->EndFinalPostprocessSettings(InitOptions);
 
+	// Apply AA jitter via HackAddTemporalAAProjectionJitter (must be after view creation).
+	if (UApertureSampler* Sampler = GetTypedOuter<UApertureSampler>())
+	{
+		const FIntPoint Resolution = ViewFamily->RenderTarget->GetSizeXY();
+		const FVector2D Jitter = Sampler->GetJitterForProjectionMatrix(Params.SampleIndex, Resolution);
+		if (!Jitter.IsNearlyZero())
+		{
+			View->ViewMatrices.HackAddTemporalAAProjectionJitter(Jitter);
+		}
+	}
+
 	// Setup view extensions
 	for (int ViewExt = 0; ViewExt < ViewFamily->ViewExtensions.Num(); ViewExt++)
 	{
diff --git a/Engine/Plugins/VirtualProduction/AccumulationDOF/Source/AccumulationDOF/Private/Shaders/AccumulationDOFShaders.cpp b/Engine/Plugins/VirtualProduction/AccumulationDOF/Source/AccumulationDOF/Private/Shaders/AccumulationDOFShaders.cpp
--- a/Engine/Plugins/VirtualProduction/AccumulationDOF/Source/AccumulationDOF/Private/Shaders/AccumulationDOFShaders.cpp
+++ b/Engine/Plugins/VirtualProduction/AccumulationDOF/Source/AccumulationDOF/Private/Shaders/AccumulationDOFShaders.cpp
@@ -105,7 +105,9 @@ void AccumulateSample(
 	PassParameters->CocRadiusToIncircleRadius      = Params.CocRadiusToIncircleRadius;
 	PassParameters->DiaphragmBladeRadius           = Params.DiaphragmBladeRadius;
 	PassParameters->DiaphragmBladeCenterOffset     = Params.DiaphragmBladeCenterOffset;
-	PassParameters->AntiAliasingWeight                       = Params.AntiAliasingWeight;
+	PassParameters->AntiAliasingWeight             = Params.AntiAliasingWeight;
+	PassParameters->SubpixelOffset                 = Params.SubpixelOffset;
+	PassParameters->TexelSize                      = Params.TexelSize;
 	PassParameters->RenderTargets[0]               = FRenderTargetBinding(OutputRDG, ERenderTargetLoadAction::ENoAction);
 
 	FIntPoint OutputSize(OutputTexture->GetSizeX(), OutputTexture->GetSizeY());
@@ -516,9 +518,8 @@ void InjectToSceneColor(
 	// MiniFont texture for preview label text rendering
 	PassParameters->MiniFontTexture = GSystemTextures.AsciiTexture ? GSystemTextures.AsciiTexture->GetRHI() : GSystemTextures.WhiteDummy->GetRHI();
 
-	// Dithering parameters (conservative strength for defense-in-depth against tone curve quantization)
 	PassParameters->FrameCounter = GFrameCounter;
-	PassParameters->DitherStrength = 1.0f / 2048.0f; // Conservative: ~0.5 LSB for f16
+	PassParameters->DitherStrength = 1.0f / 2048.0f;
 
 	PassParameters->RenderTargets[0] = FRenderTargetBinding(OutputRDG, ERenderTargetLoadAction::ENoAction);
 
diff --git a/Engine/Plugins/VirtualProduction/AccumulationDOF/Source/AccumulationDOF/Private/Shaders/AccumulationDOFShaders.h b/Engine/Plugins/VirtualProduction/AccumulationDOF/Source/AccumulationDOF/Private/Shaders/AccumulationDOFShaders.h
--- a/Engine/Plugins/VirtualProduction/AccumulationDOF/Source/AccumulationDOF/Private/Shaders/AccumulationDOFShaders.h
+++ b/Engine/Plugins/VirtualProduction/AccumulationDOF/Source/AccumulationDOF/Private/Shaders/AccumulationDOFShaders.h
@@ -82,6 +82,10 @@ class FAccumulationDOFAccumulatePS : public FGlobalShader
 		// AA reconstruction weight
 		SHADER_PARAMETER(float, AntiAliasingWeight)
 
+		// Bilinear sub-pixel splatting parameters
+		SHADER_PARAMETER(FVector2f, SubpixelOffset)
+		SHADER_PARAMETER(FVector2f, TexelSize)
+
 		RENDER_TARGET_BINDING_SLOTS()
 
 	END_SHADER_PARAMETER_STRUCT()
@@ -216,6 +220,12 @@ struct FAccumulateSampleParams
 
 	/** AA reconstruction weight */
 	float AntiAliasingWeight = 1.0f;
+
+	/** Sub-pixel offset for bilinear splatting [0,1] */
+	FVector2f SubpixelOffset = FVector2f(0.5f, 0.5f);
+
+	/** Texel size (1/width, 1/height) for UV offset calculation */
+	FVector2f TexelSize = FVector2f(0.0f, 0.0f);
 };
 
 /**
diff --git a/Engine/Plugins/VirtualProduction/AccumulationDOF/Source/AccumulationDOF/Private/Utils/AccumulationDOFMath.cpp b/Engine/Plugins/VirtualProduction/AccumulationDOF/Source/AccumulationDOF/Private/Utils/AccumulationDOFMath.cpp
--- a/Engine/Plugins/VirtualProduction/AccumulationDOF/Source/AccumulationDOF/Private/Utils/AccumulationDOFMath.cpp
+++ b/Engine/Plugins/VirtualProduction/AccumulationDOF/Source/AccumulationDOF/Private/Utils/AccumulationDOFMath.cpp
@@ -448,11 +448,8 @@ namespace AccumulationDOFMath
 
 	FVector2f GenerateHaltonJitter(int32 SampleIndex)
 	{
-		// Uses coprime bases 2 and 3
-
 		const float HaltonX = Halton(SampleIndex, 2);
 		const float HaltonY = Halton(SampleIndex, 3);
-
 		return FVector2f(HaltonX - 0.5f, HaltonY - 0.5f);
 	}
 
diff --git a/Engine/Plugins/VirtualProduction/AccumulationDOF/Source/AccumulationDOF/Public/MoviePipelineAccumulationDOFPass.h b/Engine/Plugins/VirtualProduction/AccumulationDOF/Source/AccumulationDOF/Public/MoviePipelineAccumulationDOFPass.h
--- a/Engine/Plugins/VirtualProduction/AccumulationDOF/Source/AccumulationDOF/Public/MoviePipelineAccumulationDOFPass.h
+++ b/Engine/Plugins/VirtualProduction/AccumulationDOF/Source/AccumulationDOF/Public/MoviePipelineAccumulationDOFPass.h
@@ -122,6 +122,7 @@ protected:
 	}
 
 	//~ UMoviePipelineImagePassBase Interface
+	virtual TSharedPtr<FSceneViewFamilyContext> CalculateViewFamily(FMoviePipelineRenderPassMetrics& InOutSampleState, IViewCalcPayload* OptPayload = nullptr) override;
 	virtual void GetViewShowFlags(FEngineShowFlags& OutShowFlag, EViewModeIndex& OutViewModeIndex) const override;
 	virtual UE::MoviePipeline::FImagePassCameraViewData GetCameraInfo(FMoviePipelineRenderPassMetrics& InOutSampleState, IViewCalcPayload* OptPayload = nullptr) const override;
 	virtual void AddViewExtensions(FSceneViewFamilyContext& InContext, FMoviePipelineRenderPassMetrics& InOutSampleState) override;
diff --git a/Engine/Plugins/VirtualProduction/AccumulationDOF/Source/AccumulationDOFEditor/Private/AccumulationDOFViewportExtension.cpp b/Engine/Plugins/VirtualProduction/AccumulationDOF/Source/AccumulationDOFEditor/Private/AccumulationDOFViewportExtension.cpp
--- a/Engine/Plugins/VirtualProduction/AccumulationDOF/Source/AccumulationDOFEditor/Private/AccumulationDOFViewportExtension.cpp
+++ b/Engine/Plugins/VirtualProduction/AccumulationDOF/Source/AccumulationDOFEditor/Private/AccumulationDOFViewportExtension.cpp
@@ -440,8 +440,15 @@ void FAccumulationDOFViewportExtension::TickAmortizedRendering()
 		{
 			ApertureSampler = NewObject<UApertureSampler>(GetTransientPackage());
 		}
+		else
+		{
+			ApertureSampler->Shutdown();
+		}
 
-		ApertureSampler->Initialize(Config, CameraState);
+		if (!ApertureSampler->Initialize(Config, CameraState))
+		{
+			return;
+		}
 
 		CachedViewportResolution = RenderResolution;
 		CachedWorld = World;
@@ -674,7 +681,11 @@ void FAccumulationDOFViewportExtension::CaptureOneshot()
 		ApertureSampler->Shutdown();
 	}
 
-	ApertureSampler->Initialize(Config, CameraState);
+	if (!ApertureSampler->Initialize(Config, CameraState))
+	{
+		return;
+	}
+
 	const bool bSuccess = ApertureSampler->RenderAllSamples();
 
 	if (!bSuccess)
