ADOF: Added Alt-D as preview shortcut. Also added option to use camera's ADOF component for the settings, if there is one.

Removed non-functional spatiotemporal code.

#rnx

#Synced-CL 50383198

diff --git a/Engine/Plugins/VirtualProduction/AccumulationDOF/Source/AccumulationDOF/Private/MoviePipelineAccumulationDOFPass.cpp b/Engine/Plugins/VirtualProduction/AccumulationDOF/Source/AccumulationDOF/Private/MoviePipelineAccumulationDOFPass.cpp
--- a/Engine/Plugins/VirtualProduction/AccumulationDOF/Source/AccumulationDOF/Private/MoviePipelineAccumulationDOFPass.cpp
+++ b/Engine/Plugins/VirtualProduction/AccumulationDOF/Source/AccumulationDOF/Private/MoviePipelineAccumulationDOFPass.cpp
@@ -109,19 +109,9 @@ void UMoviePipelineAccumulationDOFPass::RenderSample_GameThreadImpl(const FMovie
 {
 	UMoviePipelineImagePassBase::RenderSample_GameThreadImpl(InSampleState);
 
-	// Cache motion blur fraction for this frame.
 	CachedMotionBlurFraction = InSampleState.OutputState.TimeData.MotionBlurFraction;
 
-	// Determine if spatiotemporal mode is active
-	const bool bSpatiotemporalMode = bSpatiotemporalSampling && InSampleState.OutputState.TemporalSampleCount > 1;
-
-	// In spatiotemporal mode, only update camera state on first temporal sample to avoid
-	// resetting accumulation mid-frame when camera parameters (like focus) change
-	const bool bShouldUpdateCamera = !bSpatiotemporalMode || InSampleState.OutputState.IsFirstTemporalSample();
-	if (bShouldUpdateCamera)
-	{
-		UpdateCameraDependentState();
-	}
+	UpdateCameraDependentState();
 
 	if (!ApertureSampler || !ApertureSampler->IsInitialized())
 	{
@@ -138,10 +128,7 @@ void UMoviePipelineAccumulationDOFPass::RenderSample_GameThreadImpl(const FMovie
 		Payload.TileIndex = InSampleState.TileIndexes;
 		Payload.SceneViewIndex = 0;
 
-		// Render center sample for warmup
-
-		CurrentApertureSampleIndex = bSpatiotemporalMode ?
-			ApertureSampler->GetCenterSampleIndex() : 0;
+		CurrentApertureSampleIndex = 0;
 
 		FMoviePipelineRenderPassMetrics SampleMetrics = GetRenderPassMetricsForCamera(CameraIndex, InSampleState);
 		SampleMetrics.OutputState.CameraIndex = CameraIndex;
@@ -169,117 +156,7 @@ void UMoviePipelineAccumulationDOFPass::RenderSample_GameThreadImpl(const FMovie
 		return;
 	}
 
-	if (bSpatiotemporalMode)
-	{
-		RenderSamples_Spatiotemporal(InSampleState);
-	}
-	else
-	{
-		RenderSamples(InSampleState);
-	}
-}
-
-void UMoviePipelineAccumulationDOFPass::RenderSamples_Spatiotemporal(const FMoviePipelineRenderPassMetrics& InSampleState)
-{
-	const int32 TemporalIndex    = InSampleState.OutputState.TemporalSampleIndex;
-	const int32 TemporalCount    = InSampleState.OutputState.TemporalSampleCount;
-	const bool bIsFirstTemporal  = InSampleState.OutputState.IsFirstTemporalSample();
-	const bool bIsLastTemporal   = InSampleState.OutputState.IsLastTemporalSample();
-	const int32 ActualNumSamples = ApertureSampler->GetActualNumSamples();
-
-	if (ActualNumSamples <= 0 || TemporalCount <= 0)
-	{
-		return;
-	}
-
-	// Detect new frame
-	const int32 CurrentFrame = InSampleState.OutputState.OutputFrameNumber;
-	if (CurrentFrame != LastOutputFrameNumber)
-	{
-		LastOutputFrameNumber = CurrentFrame;
-	}
-
-	// Reset accumulation on first temporal sample of frame
-	if (bIsFirstTemporal)
-	{
-		ApertureSampler->ResetAccumulation();
-	}
-
-	// Determine aperture sample distribution across temporal samples
-
-	int32 StartApertureIndex, EndApertureIndex;
-
-	if (TemporalCount < ActualNumSamples)
-	{
-		// Spread all aperture samples across temporal samples evenly
-
-		const int32 BaseCount = ActualNumSamples / TemporalCount;
-		const int32 Remainder = ActualNumSamples % TemporalCount;
-
-		const int32 SamplesThisTemporal = (TemporalIndex < Remainder) ? (BaseCount + 1) : BaseCount;
-
-		if (TemporalIndex < Remainder)
-		{
-			StartApertureIndex = TemporalIndex * (BaseCount + 1);
-		}
-		else
-		{
-			StartApertureIndex = Remainder * (BaseCount + 1) + (TemporalIndex - Remainder) * BaseCount;
-		}
-
-		EndApertureIndex = StartApertureIndex + SamplesThisTemporal;
-	}
-	else if (TemporalCount == ActualNumSamples)
-	{
-		// 1:1 mapping
-		StartApertureIndex = TemporalIndex;
-		EndApertureIndex = TemporalIndex + 1;
-	}
-	else
-	{
-		// Oversampling - cycle through aperture samples
-
-		const int32 ApertureSampleIndex = TemporalIndex % ActualNumSamples;
-
-		StartApertureIndex = ApertureSampleIndex;
-		EndApertureIndex = ApertureSampleIndex + 1;
-	}
-
-	// Setup common rendering state
-
-	const int32 CameraIndex = 0;
-
-	FMoviePipelinePassIdentifier PassIdentifierForCurrentCamera;
-
-	PassIdentifierForCurrentCamera.Name = PassIdentifier.Name;
-	PassIdentifierForCurrentCamera.CameraName = GetCameraName(CameraIndex);
-
-	AccumulationDOF::FDeferredPassRenderStatePayload Payload;
-
-	Payload.CameraIndex = CameraIndex;
-	Payload.TileIndex = InSampleState.TileIndexes;
-	Payload.SceneViewIndex = 0;
-
-	// Render all aperture samples assigned to this temporal sample
-	for (int32 ApertureSampleIndex = StartApertureIndex; ApertureSampleIndex < EndApertureIndex; ++ApertureSampleIndex)
-	{
-		CurrentApertureSampleIndex = ApertureSampleIndex;
-
-		RenderSingleApertureSample(ApertureSampleIndex, InSampleState, Payload);
-
-		CurrentApertureSampleIndex = -1;
-	}
-
-	// Finalize accumulation on last temporal sample and submit result
-	if (bIsLastTemporal)
-	{
-		ApertureSampler->FinalizeAccumulation();
-
-		// Render injection pass and submit our accumulated result to MRQ
-		RenderInjectionAndSubmit(InSampleState, PassIdentifierForCurrentCamera);
-	}
-
-	// Note: We only submit the final accumulation
+	RenderSamples(InSampleState);
 }
 
 void UMoviePipelineAccumulationDOFPass::RenderSamples(const FMoviePipelineRenderPassMetrics& InSampleState)
@@ -512,7 +389,6 @@ void UMoviePipelineAccumulationDOFPass::RenderInjectionAndSubmit(
 	const FMoviePipelinePassIdentifier& PassIdentifierForCurrentCamera)
 {
 	const int32 CameraIndex = 0;
-	const bool bSpatiotemporalMode = bSpatiotemporalSampling && InSampleState.OutputState.TemporalSampleCount > 1;
 
 	// Get our accumulated result for injection
 	UTextureRenderTarget2D* NormalizedRT = ApertureSampler->GetAccumulatedResult();
@@ -602,18 +478,7 @@ void UMoviePipelineAccumulationDOFPass::RenderInjectionAndSubmit(
 
 	InjectionExtension.Reset();
 
-	// Submit to MRQ accumulator
-
-	FMoviePipelineRenderPassMetrics SubmissionState = InSampleState;
-
-	if (bSpatiotemporalMode)
-	{
-		// We've accumulated internally across temporal samples, so we tell MRQ this is a complete frame
-		SubmissionState.OutputState.TemporalSampleIndex = 0;
-		SubmissionState.OutputState.TemporalSampleCount = 1;
-	}
-
-	PostRendererSubmission(SubmissionState, PassIdentifierForCurrentCamera, GetOutputFileSortingOrder(), Canvas);
+	PostRendererSubmission(InSampleState, PassIdentifierForCurrentCamera, GetOutputFileSortingOrder(), Canvas);
 
 	// Save debug EXR if enabled
 	if (bSaveDebugEXR && InSampleState.OutputState.IsLastTemporalSample())
@@ -704,7 +569,6 @@ AccumulationDOF::FApertureSamplerConfig UMoviePipelineAccumulationDOFPass::Build
 		Config.TemporalHistoryMode = TemporalHistoryMode;
 	}
 
-	Config.bUnifyWithTemporalSamples = bSpatiotemporalSampling;
 	Config.ViewModeIndex = VMI_Lit;
 
 	return Config;
diff --git a/Engine/Plugins/VirtualProduction/AccumulationDOF/Source/AccumulationDOF/Private/Rendering/ApertureSampler.cpp b/Engine/Plugins/VirtualProduction/AccumulationDOF/Source/AccumulationDOF/Private/Rendering/ApertureSampler.cpp
--- a/Engine/Plugins/VirtualProduction/AccumulationDOF/Source/AccumulationDOF/Private/Rendering/ApertureSampler.cpp
+++ b/Engine/Plugins/VirtualProduction/AccumulationDOF/Source/AccumulationDOF/Private/Rendering/ApertureSampler.cpp
@@ -248,7 +248,6 @@ void UApertureSampler::Shutdown()
 	ApertureOffsetsCm.Empty();
 
 	Progress = AccumulationDOF::FApertureSamplerProgress();
-	CenterSampleIndex = 0;
 	TotalSpectralWeightPerChannel = FVector3f::ZeroVector;
 	CachedSceneFringeIntensity = 0.0f;
 
@@ -402,16 +401,6 @@ void UApertureSampler::GenerateApertureOffsets()
 		break;
 	}
 
-	// Reorder for low-discrepancy temporal traversal if spatiotemporal mode is enabled
-	if (Config.bUnifyWithTemporalSamples)
-	{
-		CenterSampleIndex = AccumulationDOFMath::ReorderSamplesForLowDiscrepancy(ApertureOffsetsCm);
-	}
-	else
-	{
-		CenterSampleIndex = 0;
-	}
-
 	Progress.TotalIterations = GetTotalIterations();
 }
 
diff --git a/Engine/Plugins/VirtualProduction/AccumulationDOF/Source/AccumulationDOF/Private/Rendering/ApertureSampler.h b/Engine/Plugins/VirtualProduction/AccumulationDOF/Source/AccumulationDOF/Private/Rendering/ApertureSampler.h
--- a/Engine/Plugins/VirtualProduction/AccumulationDOF/Source/AccumulationDOF/Private/Rendering/ApertureSampler.h
+++ b/Engine/Plugins/VirtualProduction/AccumulationDOF/Source/AccumulationDOF/Private/Rendering/ApertureSampler.h
@@ -162,15 +162,6 @@ struct FApertureSamplerConfig
 	 */
 	float ScreenPercentage = -1.0f;
 
-	//
-	// Spatiotemporal Sampling (for MRQ)
-	//
-
-	/**
-	 * Unify aperture samples with MRQ temporal samples. Experimental and may have some artifacts.
-	 */
-	bool bUnifyWithTemporalSamples = false;
-
 	//
 	// View Mode
 	//
@@ -626,23 +617,6 @@ public:
 		return TotalSpectralWeightPerChannel;
 	}
 
-	/**
-	 * Get the index of the center sample (0,0) after reordering.
-	 * Used for warmup frames to ensure TAA/Lumen priming is optimal.
-	 */
-	int32 GetCenterSampleIndex() const
-	{
-		return CenterSampleIndex;
-	}
-
-	/**
-	 * Check if spatiotemporal mode is active.
-	 */
-	bool IsSpatiotemporalMode() const
-	{
-		return Config.bUnifyWithTemporalSamples;
-	}
-
 private:
 
 	/** Setup render targets based on configuration. */
@@ -728,9 +702,6 @@ private:
 	/** Array of aperture offsets to render from per the jitter pattern */
 	TArray<FVector2f> ApertureOffsetsCm;
 
-	/** Index of center sample (0,0) in ApertureOffsetsCm after reordering */
-	int32 CenterSampleIndex = 0;
-
 	/** Axial chromatic aberration weight per channel cache */
 	FVector3f TotalSpectralWeightPerChannel = FVector3f(0.0f);
 
diff --git a/Engine/Plugins/VirtualProduction/AccumulationDOF/Source/AccumulationDOF/Private/Utils/AccumulationDOFMath.cpp b/Engine/Plugins/VirtualProduction/AccumulationDOF/Source/AccumulationDOF/Private/Utils/AccumulationDOFMath.cpp
--- a/Engine/Plugins/VirtualProduction/AccumulationDOF/Source/AccumulationDOF/Private/Utils/AccumulationDOFMath.cpp
+++ b/Engine/Plugins/VirtualProduction/AccumulationDOF/Source/AccumulationDOF/Private/Utils/AccumulationDOFMath.cpp
@@ -409,13 +409,6 @@ namespace AccumulationDOFMath
 		return Result;
 	}
 
-	int32 ReorderSamplesForLowDiscrepancy(TArray<FVector2f>& InOutOffsets)
-	{
-		// Disabled for now. Only relevant for the spatiotemporal motion blur option. 
-
-		return 0;
-	}
-
 	void ReorderSamplesForPendulum(TArray<FVector2f>& InOutOffsets)
 	{
 		// This preserves the first sample (expect center) and ends close to it which benefits temporal reprojection continuity
diff --git a/Engine/Plugins/VirtualProduction/AccumulationDOF/Source/AccumulationDOF/Private/Utils/AccumulationDOFMath.h b/Engine/Plugins/VirtualProduction/AccumulationDOF/Source/AccumulationDOF/Private/Utils/AccumulationDOFMath.h
--- a/Engine/Plugins/VirtualProduction/AccumulationDOF/Source/AccumulationDOF/Private/Utils/AccumulationDOFMath.h
+++ b/Engine/Plugins/VirtualProduction/AccumulationDOF/Source/AccumulationDOF/Private/Utils/AccumulationDOFMath.h
@@ -242,14 +242,6 @@ namespace AccumulationDOFMath
 		float FarPlane
 	);
 
-	/**
-	 * Reorder aperture samples for low-discrepancy temporal traversal.
-	 *
-	 * @param InOutOffsets - Array of aperture offsets to reorder in-place
-	 * @return Index of the original center sample (0,0) in the reordered array
-	 */
-	int32 ReorderSamplesForLowDiscrepancy(TArray<FVector2f>& InOutOffsets);
-
 	/**
 	 * Reorder samples in a pendulum (forth and back) pattern to minimize perspective jumps between frames
 	 * and also start and end near the center of the lens.
diff --git a/Engine/Plugins/VirtualProduction/AccumulationDOF/Source/AccumulationDOF/Public/MoviePipelineAccumulationDOFPass.h b/Engine/Plugins/VirtualProduction/AccumulationDOF/Source/AccumulationDOF/Public/MoviePipelineAccumulationDOFPass.h
--- a/Engine/Plugins/VirtualProduction/AccumulationDOF/Source/AccumulationDOF/Public/MoviePipelineAccumulationDOFPass.h
+++ b/Engine/Plugins/VirtualProduction/AccumulationDOF/Source/AccumulationDOF/Public/MoviePipelineAccumulationDOFPass.h
@@ -87,14 +87,6 @@ public:
 	UPROPERTY(EditAnywhere, BlueprintReadWrite, Category = "Accumulation DOF", meta = (DisplayName = "Motion Blur"))
 	bool bMotionBlur = false;
 
-	/**
-	 * Distribute aperture samples across MRQ temporal samples.
-	 * Motion blur is achieved through temporal accumulation rather than post-process.
-	 */
-	UPROPERTY(EditAnywhere, BlueprintReadWrite, Category = "Accumulation DOF",
-		meta = (DisplayName = "Spatiotemporal Sampling"))
-	bool bSpatiotemporalSampling = false;
-
 	/**
 	 * Save debug EXR files (individual samples and final accumulation).
 	 * Files are saved to the project's Saved/AccumulationDOF_MRQ_Debug folder.
@@ -174,9 +166,6 @@ private:
 	/** Whether this is the first injection render */
 	bool bFirstInjection = true;
 
-	/** Track current output frame for spatiotemporal mode reset detection */
-	int32 LastOutputFrameNumber = -1;
-
 	/** Saved LUT size before override (0 if not overridden) */
 	int32 SavedLUTSize = 0;
 
@@ -209,9 +198,6 @@ private:
 	/** Save debug EXR files. */
 	void SaveDebugEXR(int32 FrameNumber);
 
-	/** Render aperture samples distributed among temporal samples */
-	void RenderSamples_Spatiotemporal(const FMoviePipelineRenderPassMetrics& InSampleState);
-
 	/** Render all aperture samples in one temporal sample */
 	void RenderSamples(const FMoviePipelineRenderPassMetrics& InSampleState);
 
diff --git a/Engine/Plugins/VirtualProduction/AccumulationDOF/Source/AccumulationDOFEditor/Private/AccumulationDOFEditorCommands.cpp b/Engine/Plugins/VirtualProduction/AccumulationDOF/Source/AccumulationDOFEditor/Private/AccumulationDOFEditorCommands.cpp
new file mode 100644
--- /dev/null
+++ b/Engine/Plugins/VirtualProduction/AccumulationDOF/Source/AccumulationDOFEditor/Private/AccumulationDOFEditorCommands.cpp
@@ -0,0 +1,18 @@
+// Copyright Epic Games, Inc. All Rights Reserved.
+
+#include "AccumulationDOFEditorCommands.h"
+
+#define LOCTEXT_NAMESPACE "AccumulationDOFEditor"
+
+void FAccumulationDOFEditorCommands::RegisterCommands()
+{
+	UI_COMMAND(
+		ToggleAccumulate,
+		"Toggle Accumulate",
+		"Toggle accumulation depth of field preview in the active viewport",
+		EUserInterfaceActionType::ToggleButton,
+		FInputChord(EModifierKey::Alt, EKeys::D)
+	);
+}
+
+#undef LOCTEXT_NAMESPACE
diff --git a/Engine/Plugins/VirtualProduction/AccumulationDOF/Source/AccumulationDOFEditor/Private/AccumulationDOFEditorCommands.h b/Engine/Plugins/VirtualProduction/AccumulationDOF/Source/AccumulationDOFEditor/Private/AccumulationDOFEditorCommands.h
new file mode 100644
--- /dev/null
+++ b/Engine/Plugins/VirtualProduction/AccumulationDOF/Source/AccumulationDOFEditor/Private/AccumulationDOFEditorCommands.h
@@ -0,0 +1,29 @@
+// Copyright Epic Games, Inc. All Rights Reserved.
+
+#pragma once
+
+#include "CoreMinimal.h"
+#include "Framework/Commands/Commands.h"
+#include "Styling/AppStyle.h"
+
+#define LOCTEXT_NAMESPACE "AccumulationDOFEditor"
+
+class FAccumulationDOFEditorCommands : public TCommands<FAccumulationDOFEditorCommands>
+{
+public:
+	FAccumulationDOFEditorCommands()
+		: TCommands<FAccumulationDOFEditorCommands>(
+			TEXT("AccumulationDOFEditor"),
+			LOCTEXT("AccumulationDOFEditorCommandsContext", "Accumulation DOF"),
+			TEXT("LevelViewport"),
+			FAppStyle::GetAppStyleSetName()
+		)
+	{
+	}
+
+	virtual void RegisterCommands() override;
+
+	TSharedPtr<FUICommandInfo> ToggleAccumulate;
+};
+
+#undef LOCTEXT_NAMESPACE
diff --git a/Engine/Plugins/VirtualProduction/AccumulationDOF/Source/AccumulationDOFEditor/Private/AccumulationDOFEditorModule.cpp b/Engine/Plugins/VirtualProduction/AccumulationDOF/Source/AccumulationDOFEditor/Private/AccumulationDOFEditorModule.cpp
--- a/Engine/Plugins/VirtualProduction/AccumulationDOF/Source/AccumulationDOFEditor/Private/AccumulationDOFEditorModule.cpp
+++ b/Engine/Plugins/VirtualProduction/AccumulationDOF/Source/AccumulationDOFEditor/Private/AccumulationDOFEditorModule.cpp
@@ -2,22 +2,20 @@
 
 #include "AccumulationDOFEditorModule.h"
 
+#include "AccumulationDOFEditorCommands.h"
 #include "AccumulationDOFEditorSettings.h"
+#include "AccumulationDOFViewportExtension.h"
 #include "AccumulationDOFViewportManager.h"
 #include "AccumulationDOFViewportSettings.h"
 
 #include "Editor.h"
+#include "Framework/Commands/UICommandList.h"
 #include "LevelEditor.h"
 #include "LevelEditorMenuContext.h"
 #include "LevelEditorViewport.h"
 #include "SLevelViewport.h"
 #include "ToolMenus.h"
-#include "Widgets/Input/SButton.h"
-#include "Widgets/Input/SCheckBox.h"
 #include "Widgets/Input/SSpinBox.h"
-#include "Widgets/Text/STextBlock.h"
-
-#include "AccumulationDOFViewportExtension.h"
 
 #define LOCTEXT_NAMESPACE "AccumulationDOFEditor"
 
@@ -25,6 +23,8 @@ void FAccumulationDOFEditorModule::StartupModule()
 {
 	ViewportManager = MakeShared<FAccumulationDOFViewportManager>();
 
+	FAccumulationDOFEditorCommands::Register();
+
 	FCoreDelegates::OnFEngineLoopInitComplete.AddRaw(this, &FAccumulationDOFEditorModule::OnEngineLoopInitComplete);
 
 	ToolMenusStartupCallbackHandle = UToolMenus::RegisterStartupCallback(
@@ -41,6 +41,9 @@ void FAccumulationDOFEditorModule::ShutdownModule()
 
 	UnregisterViewMenuExtension();
 
+	BoundCommands.Reset();
+	FAccumulationDOFEditorCommands::Unregister();
+
 	if (GEditor)
 	{
 		GEditor->OnLevelViewportClientListChanged().RemoveAll(this);
@@ -84,6 +87,7 @@ void FAccumulationDOFEditorModule::UnregisterViewMenuExtension()
 void FAccumulationDOFEditorModule::AddAccumulationDOFControls(FToolMenuSection& InSection)
 {
 	TSharedPtr<SLevelViewport> LevelViewport = ULevelViewportContext::GetLevelViewport(InSection);
+
 	if (!LevelViewport.IsValid())
 	{
 		return;
@@ -95,31 +99,31 @@ void FAccumulationDOFEditorModule::AddAccumulationDOFControls(FToolMenuSection&
 	InSection.AddEntry(FToolMenuEntry::InitMenuEntry(
 		FName("AccumulateDOF"),
 		LOCTEXT("AccumulateDOF_Label", "Accumulate"),
-		LOCTEXT("AccumulateDOF_Tooltip", "Enable accumulation depth of field preview"),
+		LOCTEXT("AccumulateDOF_Tooltip", "Enable accumulation depth of field preview (Alt+D)"),
 		FSlateIcon(),
 		FUIAction(
 			FExecuteAction::CreateLambda([this, WeakViewport]()
 			{
 				TSharedPtr<SLevelViewport> Viewport = WeakViewport.Pin();
-				if (!Viewport.IsValid()) { return; }
-				FLevelEditorViewportClient* ViewportClient = &Viewport->GetLevelViewportClient();
-				FAccumulationDOFViewportSettings& Settings = ViewportManager->FindOrAddViewportSettings(ViewportClient);
-				Settings.bIsEnabled = !Settings.bIsEnabled;
-				EnsureViewportTrackedForPersistence(ViewportClient, Viewport);
-				ViewportManager->Unfreeze(ViewportClient);
-				if (ViewportClient->Viewport)
+
+				if (!Viewport.IsValid())
 				{
-					ViewportClient->Viewport->Invalidate();
+					return;
 				}
+
+				ToggleAccumulationForViewport(&Viewport->GetLevelViewportClient());
 			}),
 			FCanExecuteAction(),
 			FIsActionChecked::CreateLambda([this, WeakViewport]()
 			{
 				TSharedPtr<SLevelViewport> Viewport = WeakViewport.Pin();
-				if (!Viewport.IsValid()) { return false; }
-				FLevelEditorViewportClient* ViewportClient = &Viewport->GetLevelViewportClient();
-				const FAccumulationDOFViewportSettings* Settings = ViewportManager->GetViewportSettings(ViewportClient);
-				return Settings && Settings->bIsEnabled;
+
+				if (!Viewport.IsValid())
+				{
+					return false;
+				}
+
+				return IsAccumulationEnabledForViewport(&Viewport->GetLevelViewportClient());
 			})),
 		EUserInterfaceActionType::ToggleButton
 	));
@@ -130,8 +134,14 @@ void FAccumulationDOFEditorModule::AddAccumulationDOFControls(FToolMenuSection&
 		TAttribute<FText>::CreateLambda([this, WeakViewport]() -> FText
 		{
 			TSharedPtr<SLevelViewport> Viewport = WeakViewport.Pin();
-			if (!Viewport.IsValid()) { return FText::GetEmpty(); }
+
+			if (!Viewport.IsValid())
+			{
+				return FText::GetEmpty();
+			}
+
 			FLevelEditorViewportClient* ViewportClient = &Viewport->GetLevelViewportClient();
+
 			if (TSharedPtr<FAccumulationDOFViewportExtension> Extension = ViewportManager->GetViewportExtension(ViewportClient))
 			{
 				if (Extension->IsFrozen())
@@ -139,13 +149,20 @@ void FAccumulationDOFEditorModule::AddAccumulationDOFControls(FToolMenuSection&
 					return LOCTEXT("UnfreezeDOF_Label", "Unfreeze");
 				}
 			}
+
 			return LOCTEXT("OneshotDOF_Label", "One-shot");
 		}),
 		TAttribute<FText>::CreateLambda([this, WeakViewport]() -> FText
 		{
 			TSharedPtr<SLevelViewport> Viewport = WeakViewport.Pin();
-			if (!Viewport.IsValid()) { return FText::GetEmpty(); }
+
+			if (!Viewport.IsValid())
+			{
+				return FText::GetEmpty();
+			}
+
 			FLevelEditorViewportClient* ViewportClient = &Viewport->GetLevelViewportClient();
+
 			if (TSharedPtr<FAccumulationDOFViewportExtension> Extension = ViewportManager->GetViewportExtension(ViewportClient))
 			{
 				if (Extension->IsFrozen())
@@ -153,6 +170,7 @@ void FAccumulationDOFEditorModule::AddAccumulationDOFControls(FToolMenuSection&
 					return LOCTEXT("UnfreezeDOF_Tooltip", "Unfreeze the accumulated DOF result");
 				}
 			}
+
 			return LOCTEXT("OneshotDOF_Tooltip", "Capture a single high-quality DOF frame");
 		}),
 		FSlateIcon(),
@@ -160,20 +178,29 @@ void FAccumulationDOFEditorModule::AddAccumulationDOFControls(FToolMenuSection&
 			FExecuteAction::CreateLambda([this, WeakViewport]()
 			{
 				TSharedPtr<SLevelViewport> Viewport = WeakViewport.Pin();
-				if (!Viewport.IsValid()) { return; }
+
+				if (!Viewport.IsValid())
+				{
+					return;
+				}
+
 				FLevelEditorViewportClient* ViewportClient = &Viewport->GetLevelViewportClient();
+
 				if (TSharedPtr<FAccumulationDOFViewportExtension> Extension = ViewportManager->GetViewportExtension(ViewportClient))
 				{
 					if (Extension->IsFrozen())
 					{
 						ViewportManager->Unfreeze(ViewportClient);
+
 						return;
 					}
 				}
+
 				ViewportManager->CaptureOneshot(ViewportClient);
 			})),
 		EUserInterfaceActionType::Button
 	);
+
 	InSection.AddEntry(OneshotEntry);
 
 	// Settings submenu
@@ -195,6 +222,7 @@ void FAccumulationDOFEditorModule::AddAccumulationDOFSettingsSubMenu(UToolMenu*
 		FNewToolMenuDelegate::CreateLambda([this](UToolMenu* DynamicMenu)
 		{
 			TSharedPtr<SLevelViewport> LevelViewport = ULevelViewportContext::GetLevelViewport(DynamicMenu);
+
 			if (!LevelViewport.IsValid())
 			{
 				return;
@@ -202,40 +230,109 @@ void FAccumulationDOFEditorModule::AddAccumulationDOFSettingsSubMenu(UToolMenu*
 
 			FLevelEditorViewportClient* ViewportClient = &LevelViewport->GetLevelViewportClient();
 			ViewportManager->FindOrAddViewportSettings(ViewportClient);
-
 			TWeakPtr<SLevelViewport> WeakViewport = LevelViewport;
-
 			FToolMenuSection& SettingsSection = DynamicMenu->AddSection("AdvancedDOFSettings");
 
+			// Use Camera Settings checkbox
+			SettingsSection.AddEntry(FToolMenuEntry::InitMenuEntry(
+				FName("UseCameraSettings"),
+				LOCTEXT("UseCameraSettings_Label", "Use Camera Settings"),
+				LOCTEXT("UseCameraSettings_Tooltip", "Use the number of samples and DOF splat size specified the cine camera's AccumulationDOF component, if there is one."),
+				FSlateIcon(),
+				FUIAction(
+					FExecuteAction::CreateLambda([this, WeakViewport]()
+					{
+						TSharedPtr<SLevelViewport> Viewport = WeakViewport.Pin();
+
+						if (!Viewport.IsValid())
+						{
+							return;
+						}
+
+						FLevelEditorViewportClient* ViewportClient = &Viewport->GetLevelViewportClient();
+						FAccumulationDOFViewportSettings& Settings = ViewportManager->FindOrAddViewportSettings(ViewportClient);
+						Settings.bUseCameraSettings = !Settings.bUseCameraSettings;
+						EnsureViewportTrackedForPersistence(ViewportClient, Viewport);
+						ViewportManager->RestartAccumulation(ViewportClient);
+
+						if (ViewportClient->Viewport)
+						{
+							ViewportClient->Viewport->Invalidate();
+						}
+					}),
+					FCanExecuteAction(),
+					FIsActionChecked::CreateLambda([this, WeakViewport]()
+					{
+						TSharedPtr<SLevelViewport> Viewport = WeakViewport.Pin();
+
+						if (!Viewport.IsValid())
+						{
+							return false;
+						}
+
+						FLevelEditorViewportClient* ViewportClient = &Viewport->GetLevelViewportClient();
+						const FAccumulationDOFViewportSettings* Settings = ViewportManager->GetViewportSettings(ViewportClient);
+
+						return Settings && Settings->bUseCameraSettings;
+					})),
+				EUserInterfaceActionType::ToggleButton
+			));
+
 			// Number of samples
 			SettingsSection.AddEntry(FToolMenuEntry::InitWidget(
 				"NumSamplesWidget",
 				SNew(SSpinBox<int32>)
 				.MinValue(1)
 				.MaxValue(16384)
+				.IsEnabled_Lambda([this, WeakViewport]()
+				{
+					TSharedPtr<SLevelViewport> Viewport = WeakViewport.Pin();
+
+					if (!Viewport.IsValid())
+					{
+						return true;
+					}
+
+					FLevelEditorViewportClient* ViewportClient = &Viewport->GetLevelViewportClient();
+					const FAccumulationDOFViewportSettings* Settings = ViewportManager->GetViewportSettings(ViewportClient);
+
+					return !Settings || !Settings->bUseCameraSettings;
+				})
 				.Value_Lambda([this, WeakViewport]()
 				{
 					TSharedPtr<SLevelViewport> Viewport = WeakViewport.Pin();
-					if (!Viewport.IsValid()) { return 256; }
+
+					if (!Viewport.IsValid())
+					{
+						return 256;
+					}
+
 					FLevelEditorViewportClient* ViewportClient = &Viewport->GetLevelViewportClient();
 					const FAccumulationDOFViewportSettings* Settings = ViewportManager->GetViewportSettings(ViewportClient);
+
 					return Settings ? Settings->NumApertureSamples : 256;
 				})
 				.OnValueCommitted_Lambda([this, WeakViewport](int32 NewValue, ETextCommit::Type)
 				{
 					TSharedPtr<SLevelViewport> Viewport = WeakViewport.Pin();
-					if (!Viewport.IsValid()) { return; }
+
+					if (!Viewport.IsValid())
+					{
+						return;
+					}
+
 					FLevelEditorViewportClient* ViewportClient = &Viewport->GetLevelViewportClient();
 					FAccumulationDOFViewportSettings& Settings = ViewportManager->FindOrAddViewportSettings(ViewportClient);
 					Settings.NumApertureSamples = FMath::Clamp(NewValue, 1, 16384);
 					EnsureViewportTrackedForPersistence(ViewportClient, Viewport);
 					ViewportManager->RestartAccumulation(ViewportClient);
+
 					if (ViewportClient->Viewport)
 					{
 						ViewportClient->Viewport->Invalidate();
 					}
 				}),
-				LOCTEXT("NumSamplesLabel", "Samples")
+				LOCTEXT("NumSamplesLabel", "Samples   ")
 			));
 
 			// DOF Splat Size
@@ -245,23 +342,49 @@ void FAccumulationDOFEditorModule::AddAccumulationDOFSettingsSubMenu(UToolMenu*
 				.MinValue(0.0f)
 				.MaxValue(1.0f)
 				.Delta(0.01f)
+				.IsEnabled_Lambda([this, WeakViewport]()
+				{
+					TSharedPtr<SLevelViewport> Viewport = WeakViewport.Pin();
+
+					if (!Viewport.IsValid())
+					{
+						return true;
+					}
+
+					FLevelEditorViewportClient* ViewportClient = &Viewport->GetLevelViewportClient();
+					const FAccumulationDOFViewportSettings* Settings = ViewportManager->GetViewportSettings(ViewportClient);
+
+					return !Settings || !Settings->bUseCameraSettings;
+				})
 				.Value_Lambda([this, WeakViewport]()
 				{
 					TSharedPtr<SLevelViewport> Viewport = WeakViewport.Pin();
-					if (!Viewport.IsValid()) { return 0.125f; }
+
+					if (!Viewport.IsValid())
+					{
+						return 0.125f;
+					}
+
 					FLevelEditorViewportClient* ViewportClient = &Viewport->GetLevelViewportClient();
 					const FAccumulationDOFViewportSettings* Settings = ViewportManager->GetViewportSettings(ViewportClient);
+
 					return Settings ? Settings->DOFSplatSize : 0.125f;
 				})
 				.OnValueCommitted_Lambda([this, WeakViewport](float NewValue, ETextCommit::Type)
 				{
 					TSharedPtr<SLevelViewport> Viewport = WeakViewport.Pin();
-					if (!Viewport.IsValid()) { return; }
+
+					if (!Viewport.IsValid())
+					{
+						return;
+					}
+
 					FLevelEditorViewportClient* ViewportClient = &Viewport->GetLevelViewportClient();
 					FAccumulationDOFViewportSettings& Settings = ViewportManager->FindOrAddViewportSettings(ViewportClient);
 					Settings.DOFSplatSize = FMath::Clamp(NewValue, 0.0f, 1.0f);
 					EnsureViewportTrackedForPersistence(ViewportClient, Viewport);
 					ViewportManager->RestartAccumulation(ViewportClient);
+
 					if (ViewportClient->Viewport)
 					{
 						ViewportClient->Viewport->Invalidate();
@@ -287,22 +410,27 @@ void FAccumulationDOFEditorModule::OnLevelViewportClientListChanged()
 	{
 		FViewportPair& ViewportData = *ViewportIt;
 
-		if (!LevelViewportClients.ContainsByPredicate([ViewportData](const FLevelEditorViewportClient* Other)
-			{ return Other == ViewportData.Key; }))
+		bool bViewportStillExists = LevelViewportClients.ContainsByPredicate([ViewportData](const FLevelEditorViewportClient* Other)
 		{
-			// Viewport was removed so we save settings and clean up
+			return Other == ViewportData.Key;
+		});
 
-			const FAccumulationDOFViewportSettings* Settings = ViewportManager->GetViewportSettings(ViewportData.Key);
+		if (bViewportStillExists)
+		{
+			continue;
+		}
 
-			if (Settings)
-			{
-				GetMutableDefault<UAccumulationDOFLevelViewportSettings>()->SetViewportSettings(ViewportData.Value, *Settings);
-				GetMutableDefault<UAccumulationDOFLevelViewportSettings>()->SaveConfig();
-				ViewportManager->RemoveViewportSettings(ViewportData.Key);
-			}
+		// Viewport was removed, save settings and clean up
+		const FAccumulationDOFViewportSettings* Settings = ViewportManager->GetViewportSettings(ViewportData.Key);
 
-			ViewportIt.RemoveCurrent();
+		if (Settings)
+		{
+			GetMutableDefault<UAccumulationDOFLevelViewportSettings>()->SetViewportSettings(ViewportData.Value, *Settings);
+			GetMutableDefault<UAccumulationDOFLevelViewportSettings>()->SaveConfig();
+			ViewportManager->RemoveViewportSettings(ViewportData.Key);
 		}
+
+		ViewportIt.RemoveCurrent();
 	}
 
 	// Also cleanup any extensions for viewports not in ConfiguredViewports
@@ -321,31 +449,42 @@ void FAccumulationDOFEditorModule::OnEngineLoopInitComplete()
 	{
 		GEditor->OnLevelViewportClientListChanged().AddRaw(this, &FAccumulationDOFEditorModule::OnLevelViewportClientListChanged);
 	}
+
+	BindCommands();
 }
 
 void FAccumulationDOFEditorModule::SetupViewportDisplaySettings(FLevelEditorViewportClient* Client)
 {
-	if (!ConfiguredViewports.ContainsByPredicate([Client](const FViewportPair& Other)
-		{ return Other.Key == Client; }))
+	bool bAlreadyConfigured = ConfiguredViewports.ContainsByPredicate([Client](const FViewportPair& Other)
 	{
-		TSharedPtr<SLevelViewport> LevelViewport = StaticCastSharedPtr<SLevelViewport>(Client->GetEditorViewportWidget());
+		return Other.Key == Client;
+	});
 
-		if (LevelViewport.IsValid())
-		{
-			TOptional<FAccumulationDOFViewportSettings> PresetConfig =
-				GetDefault<UAccumulationDOFLevelViewportSettings>()->GetViewportSettings(LevelViewport->GetConfigKey());
-			if (PresetConfig.IsSet())
-			{
-				FAccumulationDOFViewportSettings& ViewportConfig =
-					ViewportManager->FindOrAddViewportSettings(Client);
+	if (bAlreadyConfigured)
+	{
+		return;
+	}
 
-				ViewportConfig = PresetConfig.GetValue();
+	TSharedPtr<SLevelViewport> LevelViewport = StaticCastSharedPtr<SLevelViewport>(Client->GetEditorViewportWidget());
 
-				FViewportPair NewEntry(Client, LevelViewport->GetConfigKey());
-				ConfiguredViewports.Emplace(MoveTemp(NewEntry));
-			}
-		}
+	if (!LevelViewport.IsValid())
+	{
+		return;
 	}
+
+	TOptional<FAccumulationDOFViewportSettings> PresetConfig =
+		GetDefault<UAccumulationDOFLevelViewportSettings>()->GetViewportSettings(LevelViewport->GetConfigKey());
+
+	if (!PresetConfig.IsSet())
+	{
+		return;
+	}
+
+	FAccumulationDOFViewportSettings& ViewportConfig = ViewportManager->FindOrAddViewportSettings(Client);
+	ViewportConfig = PresetConfig.GetValue();
+
+	FViewportPair NewEntry(Client, LevelViewport->GetConfigKey());
+	ConfiguredViewports.Emplace(MoveTemp(NewEntry));
 }
 
 void FAccumulationDOFEditorModule::EnsureViewportTrackedForPersistence(FLevelEditorViewportClient* Client, TSharedPtr<SLevelViewport> LevelViewport)
@@ -355,12 +494,91 @@ void FAccumulationDOFEditorModule::EnsureViewportTrackedForPersistence(FLevelEdi
 		return;
 	}
 
-	if (!ConfiguredViewports.ContainsByPredicate([Client](const FViewportPair& Other)
-		{ return Other.Key == Client; }))
+	bool bAlreadyTracked = ConfiguredViewports.ContainsByPredicate([Client](const FViewportPair& Other)
+	{
+		return Other.Key == Client;
+	});
+
+	if (bAlreadyTracked)
 	{
-		FViewportPair NewEntry(Client, LevelViewport->GetConfigKey());
-		ConfiguredViewports.Emplace(MoveTemp(NewEntry));
+		return;
 	}
+
+	FViewportPair NewEntry(Client, LevelViewport->GetConfigKey());
+	ConfiguredViewports.Emplace(MoveTemp(NewEntry));
+}
+
+void FAccumulationDOFEditorModule::ToggleAccumulationForViewport(FLevelEditorViewportClient* ViewportClient)
+{
+	if (!ViewportClient)
+	{
+		return;
+	}
+
+	TSharedPtr<SLevelViewport> LevelViewport = StaticCastSharedPtr<SLevelViewport>(ViewportClient->GetEditorViewportWidget());
+	FAccumulationDOFViewportSettings& Settings = ViewportManager->FindOrAddViewportSettings(ViewportClient);
+	Settings.bIsEnabled = !Settings.bIsEnabled;
+
+	if (LevelViewport.IsValid())
+	{
+		EnsureViewportTrackedForPersistence(ViewportClient, LevelViewport);
+	}
+
+	ViewportManager->Unfreeze(ViewportClient);
+
+	if (ViewportClient->Viewport)
+	{
+		ViewportClient->Viewport->Invalidate();
+	}
+}
+
+bool FAccumulationDOFEditorModule::IsAccumulationEnabledForViewport(FLevelEditorViewportClient* ViewportClient) const
+{
+	if (!ViewportClient)
+	{
+		return false;
+	}
+
+	const FAccumulationDOFViewportSettings* Settings = ViewportManager->GetViewportSettings(ViewportClient);
+
+	return Settings && Settings->bIsEnabled;
+}
+
+void FAccumulationDOFEditorModule::BindCommands()
+{
+	if (!FModuleManager::Get().IsModuleLoaded("LevelEditor"))
+	{
+		return;
+	}
+
+	FLevelEditorModule& LevelEditorModule = FModuleManager::GetModuleChecked<FLevelEditorModule>("LevelEditor");
+	const TSharedRef<FUICommandList>& GlobalActions = LevelEditorModule.GetGlobalLevelEditorActions();
+
+	const FAccumulationDOFEditorCommands& Commands = FAccumulationDOFEditorCommands::Get();
+
+	BoundCommands = MakeShared<FUICommandList>();
+	BoundCommands->MapAction(
+		Commands.ToggleAccumulate,
+		FExecuteAction::CreateLambda([]()
+		{
+			if (FAccumulationDOFEditorModule* Module = FModuleManager::Get().GetModulePtr<FAccumulationDOFEditorModule>("AccumulationDOFEditor"))
+			{
+				Module->ToggleAccumulationForViewport(GCurrentLevelEditingViewportClient);
+			}
+		}),
+		FCanExecuteAction(),
+		FIsActionChecked::CreateLambda([]()
+		{
+			if (FAccumulationDOFEditorModule* Module = FModuleManager::Get().GetModulePtr<FAccumulationDOFEditorModule>("AccumulationDOFEditor"))
+			{
+				return Module->IsAccumulationEnabledForViewport(GCurrentLevelEditingViewportClient);
+			}
+
+			return false;
+		})
+	);
+
+	GlobalActions->Append(BoundCommands.ToSharedRef());
 }
 
 IMPLEMENT_MODULE(FAccumulationDOFEditorModule, AccumulationDOFEditor)
diff --git a/Engine/Plugins/VirtualProduction/AccumulationDOF/Source/AccumulationDOFEditor/Private/AccumulationDOFEditorModule.h b/Engine/Plugins/VirtualProduction/AccumulationDOF/Source/AccumulationDOFEditor/Private/AccumulationDOFEditorModule.h
--- a/Engine/Plugins/VirtualProduction/AccumulationDOF/Source/AccumulationDOFEditor/Private/AccumulationDOFEditorModule.h
+++ b/Engine/Plugins/VirtualProduction/AccumulationDOF/Source/AccumulationDOFEditor/Private/AccumulationDOFEditorModule.h
@@ -8,6 +8,7 @@
 
 class FAccumulationDOFViewportManager;
 class FLevelEditorViewportClient;
+class FUICommandList;
 class UToolMenu;
 struct FAccumulationDOFViewportSettings;
 struct FToolMenuSection;
@@ -63,6 +64,15 @@ private:
 	/** Ensure viewport is tracked for settings persistence on removal */
 	void EnsureViewportTrackedForPersistence(FLevelEditorViewportClient* Client, TSharedPtr<class SLevelViewport> LevelViewport);
 
+	/** Toggle accumulation for a specific viewport client */
+	void ToggleAccumulationForViewport(FLevelEditorViewportClient* ViewportClient);
+
+	/** Check if accumulation is enabled for a specific viewport client */
+	bool IsAccumulationEnabledForViewport(FLevelEditorViewportClient* ViewportClient) const;
+
+	/** Bind keyboard shortcut commands */
+	void BindCommands();
+
 private:
 
 	/** Manages per-viewport DOF configurations and extensions */
@@ -76,4 +86,7 @@ private:
 
 	/** Delegate handle for ToolMenus startup callback */
 	FDelegateHandle ToolMenusStartupCallbackHandle;
+
+	/** Command list for keyboard shortcuts */
+	TSharedPtr<FUICommandList> BoundCommands;
 };
diff --git a/Engine/Plugins/VirtualProduction/AccumulationDOF/Source/AccumulationDOFEditor/Private/AccumulationDOFEditorSettings.h b/Engine/Plugins/VirtualProduction/AccumulationDOF/Source/AccumulationDOFEditor/Private/AccumulationDOFEditorSettings.h
--- a/Engine/Plugins/VirtualProduction/AccumulationDOF/Source/AccumulationDOFEditor/Private/AccumulationDOFEditorSettings.h
+++ b/Engine/Plugins/VirtualProduction/AccumulationDOF/Source/AccumulationDOFEditor/Private/AccumulationDOFEditorSettings.h
@@ -24,6 +24,10 @@ struct FAccumulationDOFPerViewportSettingsPair
 	UPROPERTY(config)
 	bool bIsEnabled = false;
 
+	/** Use settings from camera's AccumulationDOFComponent */
+	UPROPERTY(config)
+	bool bUseCameraSettings = false;
+
 	/** Number of aperture samples */
 	UPROPERTY(config)
 	int32 NumApertureSamples = 256;
@@ -42,6 +46,7 @@ struct FAccumulationDOFPerViewportSettingsPair
 		FAccumulationDOFViewportSettings Settings;
 
 		Settings.bIsEnabled         = bIsEnabled;
+		Settings.bUseCameraSettings = bUseCameraSettings;
 		Settings.NumApertureSamples = FMath::Max(1, NumApertureSamples);
 		Settings.DOFSplatSize       = FMath::Clamp(DOFSplatSize, 0.0f, 1.0f);
 		Settings.SamplesPerFrame    = FMath::Max(1, SamplesPerFrame);
@@ -53,6 +58,7 @@ struct FAccumulationDOFPerViewportSettingsPair
 	void FromSettings(const FAccumulationDOFViewportSettings& Settings)
 	{
 		bIsEnabled         = Settings.bIsEnabled;
+		bUseCameraSettings = Settings.bUseCameraSettings;
 		NumApertureSamples = Settings.NumApertureSamples;
 		DOFSplatSize       = Settings.DOFSplatSize;
 		SamplesPerFrame    = Settings.SamplesPerFrame;
diff --git a/Engine/Plugins/VirtualProduction/AccumulationDOF/Source/AccumulationDOFEditor/Private/AccumulationDOFViewportExtension.cpp b/Engine/Plugins/VirtualProduction/AccumulationDOF/Source/AccumulationDOFEditor/Private/AccumulationDOFViewportExtension.cpp
--- a/Engine/Plugins/VirtualProduction/AccumulationDOF/Source/AccumulationDOFEditor/Private/AccumulationDOFViewportExtension.cpp
+++ b/Engine/Plugins/VirtualProduction/AccumulationDOF/Source/AccumulationDOFEditor/Private/AccumulationDOFViewportExtension.cpp
@@ -415,6 +415,13 @@ void FAccumulationDOFViewportExtension::TickAmortizedRendering()
 
 		if (UAccumulationDOFComponent* DOFComp = GetPilotedDOFComponent())
 		{
+			// Use camera component's NumSamples and DOFSplatSize when bUseCameraSettings is enabled
+			if (Settings.bUseCameraSettings)
+			{
+				Config.NumSamples   = DOFComp->NumSamples;
+				Config.DOFSplatSize = DOFComp->DOFSplatSize;
+			}
+
 			Config.SamplingPattern                     = DOFComp->SamplingPattern;
 			Config.AxialChromaticAberrationIntensity   = DOFComp->AxialChromaticAberrationIntensity;
 			Config.AxialChromaticAberrationNumBands    = DOFComp->AxialChromaticAberrationNumBands;
@@ -657,6 +664,13 @@ void FAccumulationDOFViewportExtension::CaptureOneshot()
 	// Use component settings when available
 	if (UAccumulationDOFComponent* DOFComp = GetPilotedDOFComponent())
 	{
+		// Use camera component's NumSamples and DOFSplatSize when bUseCameraSettings is enabled
+		if (Settings.bUseCameraSettings)
+		{
+			Config.NumSamples   = DOFComp->NumSamples;
+			Config.DOFSplatSize = DOFComp->DOFSplatSize;
+		}
+
 		Config.SamplingPattern                     = DOFComp->SamplingPattern;
 		Config.AxialChromaticAberrationIntensity   = DOFComp->AxialChromaticAberrationIntensity;
 		Config.AxialChromaticAberrationNumBands    = DOFComp->AxialChromaticAberrationNumBands;
diff --git a/Engine/Plugins/VirtualProduction/AccumulationDOF/Source/AccumulationDOFEditor/Private/AccumulationDOFViewportSettings.h b/Engine/Plugins/VirtualProduction/AccumulationDOF/Source/AccumulationDOFEditor/Private/AccumulationDOFViewportSettings.h
--- a/Engine/Plugins/VirtualProduction/AccumulationDOF/Source/AccumulationDOFEditor/Private/AccumulationDOFViewportSettings.h
+++ b/Engine/Plugins/VirtualProduction/AccumulationDOF/Source/AccumulationDOFEditor/Private/AccumulationDOFViewportSettings.h
@@ -12,6 +12,9 @@ struct FAccumulationDOFViewportSettings
 	/** Accumulation DOF is enabled for this viewport */
 	bool bIsEnabled = false;
 
+	/** Use settings from camera's AccumulationDOFComponent instead of overrides below */
+	bool bUseCameraSettings = false;
+
 	/** Number of aperture samples to accumulate */
 	int32 NumApertureSamples = 256;
 
