ADOF: Fix physical camera exposure compensation when not manual exposure.

#rnx

#Synced-CL 50383198

diff --git a/Engine/Plugins/VirtualProduction/AccumulationDOF/Source/AccumulationDOF/Private/Rendering/AccumulationDOFSceneViewExtension.cpp b/Engine/Plugins/VirtualProduction/AccumulationDOF/Source/AccumulationDOF/Private/Rendering/AccumulationDOFSceneViewExtension.cpp
--- a/Engine/Plugins/VirtualProduction/AccumulationDOF/Source/AccumulationDOF/Private/Rendering/AccumulationDOFSceneViewExtension.cpp
+++ b/Engine/Plugins/VirtualProduction/AccumulationDOF/Source/AccumulationDOF/Private/Rendering/AccumulationDOFSceneViewExtension.cpp
@@ -90,11 +90,12 @@ void FAccumulationDOFSceneViewExtension::SetupView(FSceneViewFamily& InViewFamil
 
 		if (DOFSplatsFStop > UE_KINDA_SMALL_NUMBER)
 		{
-			// Compensate for the f-stop override with exposure compensation when using f-stop for physical exposure
+			// Compensate for f-stop override when using manual exposure with physical camera settings
 
 			const float CameraFStop = InView.FinalPostProcessSettings.DepthOfFieldFstop;
+			const bool bManualExposure = (InView.FinalPostProcessSettings.AutoExposureMethod == EAutoExposureMethod::AEM_Manual);
 
-			if (InView.FinalPostProcessSettings.AutoExposureApplyPhysicalCameraExposure && CameraFStop > UE_KINDA_SMALL_NUMBER)
+			if (bManualExposure && InView.FinalPostProcessSettings.AutoExposureApplyPhysicalCameraExposure && CameraFStop > UE_KINDA_SMALL_NUMBER)
 			{
 				const float ExposureCompensation = 2.0f * FMath::Log2(DOFSplatsFStop / CameraFStop);
 				InView.FinalPostProcessSettings.AutoExposureBias += ExposureCompensation;
