ADOF: Fix preview accumulate tooltip to reflect actual shortcut state.

#rnx

#Synced-CL 50463108

diff --git a/Engine/Plugins/VirtualProduction/AccumulationDOF/Source/AccumulationDOFEditor/Private/AccumulationDOFEditorModule.cpp b/Engine/Plugins/VirtualProduction/AccumulationDOF/Source/AccumulationDOFEditor/Private/AccumulationDOFEditorModule.cpp
--- a/Engine/Plugins/VirtualProduction/AccumulationDOF/Source/AccumulationDOFEditor/Private/AccumulationDOFEditorModule.cpp
+++ b/Engine/Plugins/VirtualProduction/AccumulationDOF/Source/AccumulationDOFEditor/Private/AccumulationDOFEditorModule.cpp
@@ -99,7 +99,21 @@ void FAccumulationDOFEditorModule::AddAccumulationDOFControls(FToolMenuSection&
 	InSection.AddEntry(FToolMenuEntry::InitMenuEntry(
 		FName("AccumulateDOF"),
 		LOCTEXT("AccumulateDOF_Label", "Accumulate"),
-		LOCTEXT("AccumulateDOF_Tooltip", "Enable accumulation depth of field preview (Alt+D)"),
+		TAttribute<FText>::CreateLambda([]() -> FText
+		{
+			const FAccumulationDOFEditorCommands& Commands = FAccumulationDOFEditorCommands::Get();
+
+			if (Commands.ToggleAccumulate.IsValid())
+			{
+				FText InputText = Commands.ToggleAccumulate->GetInputText();
+				if (!InputText.IsEmpty())
+				{
+					return FText::Format(LOCTEXT("AccumulateDOF_Tooltip_WithShortcut", "Enable accumulation depth of field preview ({0})"), InputText);
+				}
+			}
+
+			return LOCTEXT("AccumulateDOF_Tooltip_NoShortcut", "Enable accumulation depth of field preview (shortcut unset)");
+		}),
 		FSlateIcon(),
 		FUIAction(
 			FExecuteAction::CreateLambda([this, WeakViewport]()
