ADOF: Fix chromatic aberration not showing up in MRQ renders.

#rnx

#Synced-CL 50463108

diff --git a/Engine/Plugins/VirtualProduction/AccumulationDOF/Source/AccumulationDOF/Private/MoviePipelineAccumulationDOFPass.cpp b/Engine/Plugins/VirtualProduction/AccumulationDOF/Source/AccumulationDOF/Private/MoviePipelineAccumulationDOFPass.cpp
--- a/Engine/Plugins/VirtualProduction/AccumulationDOF/Source/AccumulationDOF/Private/MoviePipelineAccumulationDOFPass.cpp
+++ b/Engine/Plugins/VirtualProduction/AccumulationDOF/Source/AccumulationDOF/Private/MoviePipelineAccumulationDOFPass.cpp
@@ -167,7 +167,6 @@ void UMoviePipelineAccumulationDOFPass::RenderSamples(const FMoviePipelineRender
 	ApertureSampler->ResetAccumulation();
 
 	// Get camera info for pass identifier
-
 	const int32 CameraIndex = 0; // Primary camera only for now
 
 	FMoviePipelinePassIdentifier PassIdentifierForCurrentCamera;
@@ -180,17 +179,84 @@ void UMoviePipelineAccumulationDOFPass::RenderSamples(const FMoviePipelineRender
 	Payload.TileIndex = InSampleState.TileIndexes;
 	Payload.SceneViewIndex = 0;
 
-	// Render all aperture samples in a loop
-	for (int32 SampleIndex = 0; SampleIndex < ActualNumSamples; ++SampleIndex)
+	// Check for axial chromatic aberration
+	const float AxialCAIntensity = CachedDOFComponent.IsValid()
+		? CachedDOFComponent->AxialChromaticAberrationIntensity
+		: 0.0f;
+
+	const int32 EffectiveBands = ApertureSampler->GetEffectiveBands();
+
+	const bool bHasAxialCA = (AxialCAIntensity > UE_KINDA_SMALL_NUMBER) && (EffectiveBands > 1);
+
+	if (bHasAxialCA)
+	{
+		// Axial CA: render each sample with varying focus distances per spectral band
+		const FVector3f TotalSpectralWeight = ApertureSampler->GetTotalSpectralWeightPerChannel();
+
+		const FVector3f InvTotal(
+			(TotalSpectralWeight.X > UE_KINDA_SMALL_NUMBER) ? (1.0f / TotalSpectralWeight.X) : 0.0f,
+			(TotalSpectralWeight.Y > UE_KINDA_SMALL_NUMBER) ? (1.0f / TotalSpectralWeight.Y) : 0.0f,
+			(TotalSpectralWeight.Z > UE_KINDA_SMALL_NUMBER) ? (1.0f / TotalSpectralWeight.Z) : 0.0f
+		);
+
+		// Get nominal focus distance for spectral calculations
+		const float NominalFocusDistanceCm = CachedCineCameraComponent.IsValid()
+			? CachedCineCameraComponent->CurrentFocusDistance
+			: 1000.0f;
+
+		int32 IterationCount = 0;
+
+		for (int32 SampleIndex = 0; SampleIndex < ActualNumSamples; ++SampleIndex)
+		{
+			for (int32 Band = 0; Band < EffectiveBands; ++Band)
+			{
+				CurrentApertureSampleIndex = SampleIndex;
+
+				// Compute spectral parameters for this band
+				const float NormalizedSpectralPosition = (EffectiveBands == 1)
+					? 0.5f
+					: static_cast<float>(Band) / static_cast<float>(EffectiveBands - 1);
+
+				CurrentFocusDistanceOverride = ApertureSampler->ComputeSpectralFocusDistance(
+					NominalFocusDistanceCm, NormalizedSpectralPosition);
+
+				const FVector3f RawWeight = ApertureSampler->ComputeSpectralWeight(NormalizedSpectralPosition);
+				CurrentAxialCAWeight = RawWeight * InvTotal;
+
+				RenderSingleApertureSample(SampleIndex, InSampleState, Payload);
+
+				IterationCount++;
+
+				// Release memory periodically
+				if ((IterationCount % AccumulationDOFUtils::RenderingMemoryFlushBatchSize) == 0)
+				{
+					AccumulationDOFUtils::ReleaseRenderingMemory();
+				}
+			}
+		}
+	}
+	else
 	{
-		CurrentApertureSampleIndex = SampleIndex;
-		RenderSingleApertureSample(SampleIndex, InSampleState, Payload);
+		// No axial CA: standard rendering with uniform spectral weight
+		CurrentFocusDistanceOverride = -1.0f;
+		CurrentAxialCAWeight = FVector3f(1.0f);
+
+		for (int32 SampleIndex = 0; SampleIndex < ActualNumSamples; ++SampleIndex)
+		{
+			CurrentApertureSampleIndex = SampleIndex;
+			RenderSingleApertureSample(SampleIndex, InSampleState, Payload);
+		}
 	}
 
-	// Reset sample index and capture extension after all samples rendered
+	// Reset state after all samples rendered
 	CurrentApertureSampleIndex = -1;
+	CurrentFocusDistanceOverride = -1.0f;
+	CurrentAxialCAWeight = FVector3f(1.0f);
 	CaptureExtension.Reset();
 
+	// Update camera state with captured SceneFringeIntensity for lateral CA
+	ApertureSampler->UpdateCameraState(BuildCameraState());
+
 	// Finalize accumulation
 	ApertureSampler->FinalizeAccumulation();
 
@@ -371,9 +437,9 @@ void UMoviePipelineAccumulationDOFPass::RenderSingleApertureSample(
 		UKismetRenderingLibrary::ExportRenderTarget(this, SampleRT, SamplesDir, Filename);
 	}
 
-	// Accumulate
+	// Accumulate with spectral weight (set per band for axial CA)
 	const FVector2f ApertureOffset = ApertureSampler->GetApertureOffset(SampleIndex);
-	ApertureSampler->AccumulateExternalSample(SampleIndex, ApertureOffset);
+	ApertureSampler->AccumulateExternalSample(SampleIndex, ApertureOffset, CurrentAxialCAWeight);
 
 	// Release memory periodically
 	if (((SampleIndex + 1) % AccumulationDOFUtils::RenderingMemoryFlushBatchSize) == 0)
@@ -777,7 +843,7 @@ UE::MoviePipeline::FImagePassCameraViewData UMoviePipelineAccumulationDOFPass::G
 
 		FMatrix ProjectionMatrix;
 
-		if (ApertureSampler->ComputeProjectionMatrix(CurrentApertureSampleIndex, -1.0f, ProjectionMatrix))
+		if (ApertureSampler->ComputeProjectionMatrix(CurrentApertureSampleIndex, CurrentFocusDistanceOverride, ProjectionMatrix))
 		{
 			CameraData.bUseCustomProjectionMatrix = true;
 			CameraData.CustomProjectionMatrix = ProjectionMatrix;
diff --git a/Engine/Plugins/VirtualProduction/AccumulationDOF/Source/AccumulationDOF/Private/Rendering/ApertureSampler.cpp b/Engine/Plugins/VirtualProduction/AccumulationDOF/Source/AccumulationDOF/Private/Rendering/ApertureSampler.cpp
--- a/Engine/Plugins/VirtualProduction/AccumulationDOF/Source/AccumulationDOF/Private/Rendering/ApertureSampler.cpp
+++ b/Engine/Plugins/VirtualProduction/AccumulationDOF/Source/AccumulationDOF/Private/Rendering/ApertureSampler.cpp
@@ -249,7 +249,7 @@ void UApertureSampler::Shutdown()
 
 	Progress = AccumulationDOF::FApertureSamplerProgress();
 	TotalSpectralWeightPerChannel = FVector3f::ZeroVector;
-	CachedSceneFringeIntensity = 0.0f;
+	CameraState.SceneFringeIntensity = 0.0f;
 
 	if (SampleRT || AccumA || AccumB || PrefilterRT || PostProcessRT)
 	{
@@ -596,7 +596,7 @@ bool UApertureSampler::RenderAllSamples()
 
 				AccumulateSampleInternal(SampleIdx, SampleApertureOffsetCm, SpectralWeight);
 
-				++IterationCount;
+				IterationCount++;
 				Progress.CurrentIteration = IterationCount;
 				NotifyProgress();
 				SlowTask.EnterProgressFrame(1);
@@ -804,7 +804,7 @@ void UApertureSampler::RenderSample(int32 SampleIndex, const FVector2f& Aperture
 	float CapturedSceneFringe = 0.0f;
 	if (Renderer->RenderSample(Params, SampleRT, CineCam, &CapturedSceneFringe))
 	{
-		CachedSceneFringeIntensity = CapturedSceneFringe;
+		CameraState.SceneFringeIntensity = CapturedSceneFringe;
 	}
 }
 
@@ -1312,7 +1312,7 @@ void UApertureSampler::ApplySpectralLateralCA()
 		return;
 	}
 
-	const float LateralCAIntensity = CachedSceneFringeIntensity;
+	const float LateralCAIntensity = CameraState.SceneFringeIntensity;
 	const bool bDoSpectralCA = Config.bSpectralLateralChromaticAberration && (LateralCAIntensity > UE_KINDA_SMALL_NUMBER);
 
 	if (!bDoSpectralCA || !PrefilterRT || !AccumB)
diff --git a/Engine/Plugins/VirtualProduction/AccumulationDOF/Source/AccumulationDOF/Private/Rendering/ApertureSampler.h b/Engine/Plugins/VirtualProduction/AccumulationDOF/Source/AccumulationDOF/Private/Rendering/ApertureSampler.h
--- a/Engine/Plugins/VirtualProduction/AccumulationDOF/Source/AccumulationDOF/Private/Rendering/ApertureSampler.h
+++ b/Engine/Plugins/VirtualProduction/AccumulationDOF/Source/AccumulationDOF/Private/Rendering/ApertureSampler.h
@@ -723,9 +723,6 @@ private:
 	/** Whether cached change detection values are valid */
 	bool bLastChangeDetectionValid = false;
 
-	/** Cached lateral chromatic aberration value */
-	float CachedSceneFringeIntensity = 0.0f;
-
 	/** Internal renderer. */
 	UPROPERTY(Transient)
 	TObjectPtr<UApertureSamplingRenderer> Renderer;
diff --git a/Engine/Plugins/VirtualProduction/AccumulationDOF/Source/AccumulationDOF/Public/MoviePipelineAccumulationDOFPass.h b/Engine/Plugins/VirtualProduction/AccumulationDOF/Source/AccumulationDOF/Public/MoviePipelineAccumulationDOFPass.h
--- a/Engine/Plugins/VirtualProduction/AccumulationDOF/Source/AccumulationDOF/Public/MoviePipelineAccumulationDOFPass.h
+++ b/Engine/Plugins/VirtualProduction/AccumulationDOF/Source/AccumulationDOF/Public/MoviePipelineAccumulationDOFPass.h
@@ -131,6 +131,12 @@ private:
 	/** Current sample being rendered */
 	mutable int32 CurrentApertureSampleIndex = -1;
 
+	/** Focus distance override for axial CA (negative means no override) */
+	mutable float CurrentFocusDistanceOverride = -1.0f;
+
+	/** RGB weight for current spectral band in axial chromatic aberration */
+	FVector3f CurrentAxialCAWeight = FVector3f(1.0f);
+
 	/** Cached reference to CineCameraComponent for SVE settings. */
 	UPROPERTY(Transient)
 	TWeakObjectPtr<UCineCameraComponent> CachedCineCameraComponent;
