ADOF: FIx not taking custom clipping plane into account.

#rnx
#Synced-CL 50463108

diff --git a/Engine/Plugins/VirtualProduction/AccumulationDOF/Source/AccumulationDOF/Private/Rendering/ApertureSampler.cpp b/Engine/Plugins/VirtualProduction/AccumulationDOF/Source/AccumulationDOF/Private/Rendering/ApertureSampler.cpp
--- a/Engine/Plugins/VirtualProduction/AccumulationDOF/Source/AccumulationDOF/Private/Rendering/ApertureSampler.cpp
+++ b/Engine/Plugins/VirtualProduction/AccumulationDOF/Source/AccumulationDOF/Private/Rendering/ApertureSampler.cpp
@@ -975,7 +975,10 @@ bool UApertureSampler::ComputeProjectionMatrix(
 		}
 	}
 
-	const float NearCm = GNearClippingPlane;
+	FMinimalViewInfo CameraView;
+	CineCam->GetCameraView(0.0f, CameraView);
+
+	const float NearCm = CameraView.GetFinalPerspectiveNearClipPlane();
 	const float FarCm = NearCm;  // Infinite far plane
 
 	// Apply overscan to sensor dimensions (expands frustum to capture overscanned region)
