ADOF: Rename spherical aberration parametrers. Add coma aberration. Fix not taking into account sensor offsets (mm).

#rnx

#Synced-CL 51370176

diff --git a/Engine/Plugins/VirtualProduction/AccumulationDOF/Shaders/Private/AccumulationDOFAccumulate.usf b/Engine/Plugins/VirtualProduction/AccumulationDOF/Shaders/Private/AccumulationDOFAccumulate.usf
--- a/Engine/Plugins/VirtualProduction/AccumulationDOF/Shaders/Private/AccumulationDOFAccumulate.usf
+++ b/Engine/Plugins/VirtualProduction/AccumulationDOF/Shaders/Private/AccumulationDOFAccumulate.usf
@@ -75,6 +75,9 @@ float AntiAliasingWeight;
 float2 SubpixelOffset;
 float2 TexelSize;
 
+// Coma aberration strength (Seidel W131)
+float ComaStrength;
+
 // Bokeh shape constants (matching DiaphragmDOF.h EBokehShape)
 #define BOKEH_SHAPE_CIRCLE          0
 #define BOKEH_SHAPE_STRAIGHT_BLADES 1
@@ -357,6 +360,48 @@ float ComputeBarrelOcclusion(float2 BokehCenterUV, float2 ApertureOffsetCm, floa
 }
 
 
+// Compute coma aberration UV shift (Seidel W131)
+//
+// @param ScreenUV  Screen UV [0,1] of the pixel
+// @param NormPupil Normalized aperture position [-1,1]
+//
+// @return UV offset to apply when sampling
+//
+float2 ComputeComaShift(float2 ScreenUV, float2 NormPupil)
+{
+	if (ComaStrength <= 0.0 || FocalLengthCm <= 0.0 || ApertureRadiusCm <= 0.0001)
+	{
+		return float2(0.0, 0.0);
+	}
+
+	// UV to sensor space
+	float2 ScreenNDC = ScreenUV * 2.0 - 1.0;
+	ScreenNDC.y *= -1.0;
+	float2 FieldPosCm = ScreenNDC * SensorHalfSizeCm;
+	float FieldHeightCm = length(FieldPosCm);
+	float FieldNorm = FieldHeightCm / max(length(SensorHalfSizeCm), 0.0001);
+
+	if (FieldNorm < 0.0001)
+	{
+		return float2(0.0, 0.0);
+	}
+
+	// Decompose pupil into tangential and sagittal components
+	float2 FieldDir = FieldPosCm / FieldHeightCm;
+	float2 SagDir = float2(-FieldDir.y, FieldDir.x);
+	float PupilT = dot(NormPupil, FieldDir);
+	float PupilS = dot(NormPupil, SagDir);
+
+	// Seidel W131 ray aberration
+	float ShiftT = FieldNorm * (PupilS * PupilS + 3.0 * PupilT * PupilT);
+	float ShiftS = FieldNorm * 2.0 * PupilS * PupilT;
+	float2 ShiftDir = ShiftT * FieldDir + ShiftS * SagDir;
+
+	// Convert to UV offset
+	return float2(-ShiftDir.x, ShiftDir.y) * ComaStrength / (2.0 * SensorHalfSizeCm);
+}
+
+
 void MainVS(
 	in  uint   VertexId    : SV_VertexID,
 	out float4 OutPosition : SV_POSITION,
@@ -376,13 +421,14 @@ void MainPS(
 {
 	float4 PrevAccum = PrevAccumTexture.Sample(PrevAccumSampler, UV);
 
-	// Bilinear subpixel splatting
-	float2 GatherOffset = (SubpixelOffset - 0.5) * TexelSize;
-	float4 Sample = SampleTexture.Sample(SampleSampler, UV + GatherOffset);
-
 	// Map aperture offset to normalized disk coordinates [-1, 1]
 	float2 NormDisk = ApertureOffsetCm / max(ApertureRadiusCm, 0.0001);
 
+	// Bilinear subpixel splatting + coma aberration UV shift
+	float2 GatherOffset = (SubpixelOffset - 0.5) * TexelSize;
+	float2 ComaOffset = ComputeComaShift(UV, NormDisk);
+	float4 Sample = SampleTexture.Sample(SampleSampler, UV + GatherOffset + ComaOffset);
+
 	// Apply Petzval field curvature distortion based on screen position
 	float2 DistortedDisk = ApplyPetzvalDistortion(UV, NormDisk);
 
diff --git a/Engine/Plugins/VirtualProduction/AccumulationDOF/Source/AccumulationDOF/Private/MoviePipelineAccumulationDOFPass.cpp b/Engine/Plugins/VirtualProduction/AccumulationDOF/Source/AccumulationDOF/Private/MoviePipelineAccumulationDOFPass.cpp
--- a/Engine/Plugins/VirtualProduction/AccumulationDOF/Source/AccumulationDOF/Private/MoviePipelineAccumulationDOFPass.cpp
+++ b/Engine/Plugins/VirtualProduction/AccumulationDOF/Source/AccumulationDOF/Private/MoviePipelineAccumulationDOFPass.cpp
@@ -615,7 +615,8 @@ AccumulationDOF::FApertureSamplerConfig UMoviePipelineAccumulationDOFPass::Build
 		Config.AxialChromaticAberrationIntensity   = CachedDOFComponent->AxialChromaticAberrationIntensity;
 		Config.AxialChromaticAberrationNumBands    = CachedDOFComponent->AxialChromaticAberrationNumBands;
 		Config.bSpectralLateralChromaticAberration = CachedDOFComponent->bSpectralLateralChromaticAberration;
-		Config.W040                                = CachedDOFComponent->W040;
+		Config.SphericalAberration                 = CachedDOFComponent->SphericalAberration;
+		Config.ComaAberration                      = CachedDOFComponent->ComaAberration;
 		Config.BokehTexture                        = CachedDOFComponent->BokehTexture;
 		Config.bEnableBokehTexture                 = CachedDOFComponent->bEnableBokehTexture;
 		Config.WeightChannel                       = CachedDOFComponent->WeightChannel;
diff --git a/Engine/Plugins/VirtualProduction/AccumulationDOF/Source/AccumulationDOF/Private/Rendering/ApertureSampler.cpp b/Engine/Plugins/VirtualProduction/AccumulationDOF/Source/AccumulationDOF/Private/Rendering/ApertureSampler.cpp
--- a/Engine/Plugins/VirtualProduction/AccumulationDOF/Source/AccumulationDOF/Private/Rendering/ApertureSampler.cpp
+++ b/Engine/Plugins/VirtualProduction/AccumulationDOF/Source/AccumulationDOF/Private/Rendering/ApertureSampler.cpp
@@ -837,10 +837,10 @@ FApertureSampleParams UApertureSampler::ComputeSampleParams(int32 SampleIndex, f
 	float FocusDistanceCm = (FocusDistanceOverride >= 0.0f) ? FocusDistanceOverride : CineCam->CurrentFocusDistance;
 
 	// Apply spherical aberration if enabled
-	if (Config.W040 > UE_KINDA_SMALL_NUMBER)
+	if (Config.SphericalAberration > UE_KINDA_SMALL_NUMBER)
 	{
 		const float DeltaZCm = AccumulationDOFMath::ComputeSphericalAberrationShift(
-			ApertureOffsetCm, ApertureRadiusCm, FocalLengthCm, Config.W040
+			ApertureOffsetCm, ApertureRadiusCm, FocalLengthCm, Config.SphericalAberration
 		);
 		FocusDistanceCm += DeltaZCm;
 		FocusDistanceCm = FMath::Max(FocusDistanceCm, 1.0f);
@@ -942,10 +942,10 @@ bool UApertureSampler::ComputeProjectionMatrix(
 	float FocusDistanceCm = (FocusDistanceOverride >= 0.0f) ? FocusDistanceOverride : CineCam->CurrentFocusDistance;
 
 	// Apply spherical aberration
-	if (Config.W040 > UE_KINDA_SMALL_NUMBER)
+	if (Config.SphericalAberration > UE_KINDA_SMALL_NUMBER)
 	{
 		const float DeltaZCm = AccumulationDOFMath::ComputeSphericalAberrationShift(
-			ApertureOffsetCm, ApertureRadiusCm, FocalLengthCm, Config.W040
+			ApertureOffsetCm, ApertureRadiusCm, FocalLengthCm, Config.SphericalAberration
 		);
 		FocusDistanceCm += DeltaZCm;
 		FocusDistanceCm = FMath::Max(FocusDistanceCm, 1.0f);
@@ -990,6 +990,20 @@ bool UApertureSampler::ComputeProjectionMatrix(
 	FVector LocalSensorCorners[4];
 	AccumulationDOFMath::ComputeSensorCorners(EffectiveSensorWidthCm, SensorHeightCm, FocalLengthCm, LocalSensorCorners);
 
+	// Apply filmback sensor offset
+
+	const float SensorOffsetXCm = AccumulationDOFMath::MmToCm(CineCam->Filmback.SensorHorizontalOffset);
+	const float SensorOffsetYCm = AccumulationDOFMath::MmToCm(CineCam->Filmback.SensorVerticalOffset);
+
+	if (FMath::Abs(SensorOffsetXCm) > UE_KINDA_SMALL_NUMBER || FMath::Abs(SensorOffsetYCm) > UE_KINDA_SMALL_NUMBER)
+	{
+		for (int32 CornerIdx = 0; CornerIdx < 4; ++CornerIdx)
+		{
+			LocalSensorCorners[CornerIdx].X += SensorOffsetXCm;
+			LocalSensorCorners[CornerIdx].Y += SensorOffsetYCm;
+		}
+	}
+
 	// Compute focus plane corners in view space
 	FVector FocusCorners_ViewSpace[4];
 	AccumulationDOFMath::ComputeFocusPlaneCorners(LocalSensorCorners, FocusDistanceCm, FocusCorners_ViewSpace);
@@ -1140,6 +1154,7 @@ void UApertureSampler::AccumulateSampleInternal(
 	AccumParams.FocalLengthCm = CamParams.FocalLengthCm;
 	AccumParams.FocusDistanceCm = CamParams.FocusDistanceCm;
 	AccumParams.SensorHalfSizeCm = FVector2f(CamParams.SensorWidthCm * 0.5f, CamParams.SensorHeightCm * 0.5f);
+	AccumParams.ComaStrength = Config.ComaAberration * 0.1f;
 
 	AccumParams.bBokehEnabled = bBokehEnabled;
 	AccumParams.BokehWeightChannel = static_cast<int32>(Config.WeightChannel);
diff --git a/Engine/Plugins/VirtualProduction/AccumulationDOF/Source/AccumulationDOF/Private/Rendering/ApertureSampler.h b/Engine/Plugins/VirtualProduction/AccumulationDOF/Source/AccumulationDOF/Private/Rendering/ApertureSampler.h
--- a/Engine/Plugins/VirtualProduction/AccumulationDOF/Source/AccumulationDOF/Private/Rendering/ApertureSampler.h
+++ b/Engine/Plugins/VirtualProduction/AccumulationDOF/Source/AccumulationDOF/Private/Rendering/ApertureSampler.h
@@ -98,11 +98,14 @@ struct FApertureSamplerConfig
 	bool bSpectralLateralChromaticAberration = true;
 
 	//
-	// Spherical Aberration
+	// Monochromatic Aberrations
 	//
 
-	/** Primary spherical aberration coefficient W040 (in cm). */
-	float W040 = 0.0f;
+	/** Primary spherical aberration coefficient (Seidel W040) in cm. */
+	float SphericalAberration = 0.0f;
+
+	/** Coma aberration strength (Seidel W131). Normalized and scaled by 0.1 before passing to shader. */
+	float ComaAberration = 0.0f;
 
 	//
 	// Bokeh Settings
diff --git a/Engine/Plugins/VirtualProduction/AccumulationDOF/Source/AccumulationDOF/Private/Shaders/AccumulationDOFShaders.cpp b/Engine/Plugins/VirtualProduction/AccumulationDOF/Source/AccumulationDOF/Private/Shaders/AccumulationDOFShaders.cpp
--- a/Engine/Plugins/VirtualProduction/AccumulationDOF/Source/AccumulationDOF/Private/Shaders/AccumulationDOFShaders.cpp
+++ b/Engine/Plugins/VirtualProduction/AccumulationDOF/Source/AccumulationDOF/Private/Shaders/AccumulationDOFShaders.cpp
@@ -108,6 +108,7 @@ void AccumulateSample(
 	PassParameters->AntiAliasingWeight             = Params.AntiAliasingWeight;
 	PassParameters->SubpixelOffset                 = Params.SubpixelOffset;
 	PassParameters->TexelSize                      = Params.TexelSize;
+	PassParameters->ComaStrength                   = Params.ComaStrength;
 	PassParameters->RenderTargets[0]               = FRenderTargetBinding(OutputRDG, ERenderTargetLoadAction::ENoAction);
 
 	FIntPoint OutputSize(OutputTexture->GetSizeX(), OutputTexture->GetSizeY());
diff --git a/Engine/Plugins/VirtualProduction/AccumulationDOF/Source/AccumulationDOF/Private/Shaders/AccumulationDOFShaders.h b/Engine/Plugins/VirtualProduction/AccumulationDOF/Source/AccumulationDOF/Private/Shaders/AccumulationDOFShaders.h
--- a/Engine/Plugins/VirtualProduction/AccumulationDOF/Source/AccumulationDOF/Private/Shaders/AccumulationDOFShaders.h
+++ b/Engine/Plugins/VirtualProduction/AccumulationDOF/Source/AccumulationDOF/Private/Shaders/AccumulationDOFShaders.h
@@ -86,6 +86,9 @@ class FAccumulationDOFAccumulatePS : public FGlobalShader
 		SHADER_PARAMETER(FVector2f, SubpixelOffset)
 		SHADER_PARAMETER(FVector2f, TexelSize)
 
+		// Coma aberration strength (Seidel W131), normalized.
+		SHADER_PARAMETER(float, ComaStrength)
+
 		RENDER_TARGET_BINDING_SLOTS()
 
 	END_SHADER_PARAMETER_STRUCT()
@@ -226,6 +229,9 @@ struct FAccumulateSampleParams
 
 	/** Texel size (1/width, 1/height) for UV offset calculation */
 	FVector2f TexelSize = FVector2f(0.0f, 0.0f);
+
+	/** Coma aberration strength (Seidel W131), normalized */
+	float ComaStrength = 0.0f;
 };
 
 /**
diff --git a/Engine/Plugins/VirtualProduction/AccumulationDOF/Source/AccumulationDOF/Private/Utils/AccumulationDOFMath.h b/Engine/Plugins/VirtualProduction/AccumulationDOF/Source/AccumulationDOF/Private/Utils/AccumulationDOFMath.h
--- a/Engine/Plugins/VirtualProduction/AccumulationDOF/Source/AccumulationDOF/Private/Utils/AccumulationDOFMath.h
+++ b/Engine/Plugins/VirtualProduction/AccumulationDOF/Source/AccumulationDOF/Private/Utils/AccumulationDOFMath.h
@@ -63,17 +63,17 @@ namespace AccumulationDOFMath
 	 * @param ApertureOffsetCm - 2D offset on the aperture (in cm)
 	 * @param ApertureRadiusCm - Physical aperture radius (in cm)
 	 * @param FocalLengthCm    - Focal length (in cm)
-	 * @param W040Cm           - Primary spherical aberration coefficient W040 (in cm)
-	 * 
+	 * @param SphericalAberrationCm - Primary spherical aberration coefficient (Seidel W040) in cm
+	 *
 	 * @return Additive shift to apply to focus distance
 	 */
 	FORCEINLINE float ComputeSphericalAberrationShift(
 		const FVector2f& ApertureOffsetCm,
 		float ApertureRadiusCm,
 		float FocalLengthCm,
-		float W040Cm)
+		float SphericalAberrationCm)
 	{
-		if (W040Cm <= 0.0f || ApertureRadiusCm <= KINDA_SMALL_NUMBER || FocalLengthCm <= KINDA_SMALL_NUMBER)
+		if (FMath::IsNearlyZero(SphericalAberrationCm) || ApertureRadiusCm <= KINDA_SMALL_NUMBER || FocalLengthCm <= KINDA_SMALL_NUMBER)
 		{
 			return 0.0f;
 		}
@@ -92,7 +92,7 @@ namespace AccumulationDOFMath
 
 		// dZ = -(4 / (n' * NA^2)) * W040 * rho^2
 		// Positive W040 -> negative shift (marginal rays focus closer than paraxial)
-		const float DeltaZCm = -(4.0f / (n_image * NASquared)) * W040Cm * RhoSquared;
+		const float DeltaZCm = -(4.0f / (n_image * NASquared)) * SphericalAberrationCm * RhoSquared;
 
 		return DeltaZCm;
 	}
diff --git a/Engine/Plugins/VirtualProduction/AccumulationDOF/Source/AccumulationDOF/Public/AccumulationDOFComponent.h b/Engine/Plugins/VirtualProduction/AccumulationDOF/Source/AccumulationDOF/Public/AccumulationDOFComponent.h
--- a/Engine/Plugins/VirtualProduction/AccumulationDOF/Source/AccumulationDOF/Public/AccumulationDOFComponent.h
+++ b/Engine/Plugins/VirtualProduction/AccumulationDOF/Source/AccumulationDOF/Public/AccumulationDOFComponent.h
@@ -106,10 +106,18 @@ public:
 	bool bSpectralLateralChromaticAberration = true;
 
 	/**
-	 * Primary spherical aberration coefficient W040 (in cm).
-	 * Based on Seidel model: dZ = -(4 / (n' * NA^2)) * W040 * rho^2
+	 * Primary spherical aberration coefficient (Seidel W040) in cm.
+	 * Positive values cause marginal rays to focus closer than paraxial rays.
 	 */
-	UPROPERTY(EditAnywhere, BlueprintReadWrite, Category = "Spherical Aberration", meta = (ClampMin = "0.0", UIMin = "0.0", UIMax = "100.0"))
-	float W040 = 0.0f;
+	UPROPERTY(EditAnywhere, BlueprintReadWrite, Category = "Monochromatic Aberrations", meta = (ClampMin = "0.0", UIMin = "0.0", UIMax = "100.0"))
+	float SphericalAberration = 0.0f;
+
+	/**
+	 * Coma aberration strength (Seidel W131).
+	 * Off-axis points develop comet-like tails pointing away from the optical axis.
+	 * The effect increases with distance from the image center.
+	 */
+	UPROPERTY(EditAnywhere, BlueprintReadWrite, Category = "Monochromatic Aberrations", meta = (ClampMin = "0.0", UIMin = "0.0", UIMax = "1.0"))
+	float ComaAberration = 0.0f;
 
 };
diff --git a/Engine/Plugins/VirtualProduction/AccumulationDOF/Source/AccumulationDOFEditor/Private/AccumulationDOFViewportExtension.cpp b/Engine/Plugins/VirtualProduction/AccumulationDOF/Source/AccumulationDOFEditor/Private/AccumulationDOFViewportExtension.cpp
--- a/Engine/Plugins/VirtualProduction/AccumulationDOF/Source/AccumulationDOFEditor/Private/AccumulationDOFViewportExtension.cpp
+++ b/Engine/Plugins/VirtualProduction/AccumulationDOF/Source/AccumulationDOFEditor/Private/AccumulationDOFViewportExtension.cpp
@@ -426,7 +426,8 @@ void FAccumulationDOFViewportExtension::TickAmortizedRendering()
 			Config.AxialChromaticAberrationIntensity   = DOFComp->AxialChromaticAberrationIntensity;
 			Config.AxialChromaticAberrationNumBands    = DOFComp->AxialChromaticAberrationNumBands;
 			Config.bSpectralLateralChromaticAberration = DOFComp->bSpectralLateralChromaticAberration;
-			Config.W040                                = DOFComp->W040;
+			Config.SphericalAberration                 = DOFComp->SphericalAberration;
+			Config.ComaAberration                      = DOFComp->ComaAberration;
 			Config.BokehTexture                        = DOFComp->BokehTexture;
 			Config.bEnableBokehTexture                 = DOFComp->bEnableBokehTexture;
 			Config.WeightChannel                       = DOFComp->WeightChannel;
@@ -675,7 +676,8 @@ void FAccumulationDOFViewportExtension::CaptureOneshot()
 		Config.AxialChromaticAberrationIntensity   = DOFComp->AxialChromaticAberrationIntensity;
 		Config.AxialChromaticAberrationNumBands    = DOFComp->AxialChromaticAberrationNumBands;
 		Config.bSpectralLateralChromaticAberration = DOFComp->bSpectralLateralChromaticAberration;
-		Config.W040                                = DOFComp->W040;
+		Config.SphericalAberration                 = DOFComp->SphericalAberration;
+		Config.ComaAberration                      = DOFComp->ComaAberration;
 		Config.BokehTexture                        = DOFComp->BokehTexture;
 		Config.bEnableBokehTexture                 = DOFComp->bEnableBokehTexture;
 		Config.WeightChannel                       = DOFComp->WeightChannel;
