MoviePipeline: Added Accumulation DOF support to MRG via the Deferred Renderer node, and added a new Accumulation DOF Modifier node.

This implementation is different than the MRQ eqivalent (which will be removed). Instead of adding an "Accumulation DOF" node to the graph, the Deferred Renderer node will automatically change to Accumulation DOF mode if the camera that is being rendered has an AccumulationDOF component on it. This component is what drives behavior of Accumulation DOF. If the settings on this component need to be tweaked, the Accumulation DOF Modifier node can be added to the graph (which includes the ability to disable the component entirely).

#Synced-CL 51651671

diff --git a/Engine/Plugins/MovieScene/MovieRenderPipeline/Source/MovieRenderPipelineCore/Private/Graph/MovieGraphDefaultRenderer.cpp b/Engine/Plugins/MovieScene/MovieRenderPipeline/Source/MovieRenderPipelineCore/Private/Graph/MovieGraphDefaultRenderer.cpp
--- a/Engine/Plugins/MovieScene/MovieRenderPipeline/Source/MovieRenderPipelineCore/Private/Graph/MovieGraphDefaultRenderer.cpp
+++ b/Engine/Plugins/MovieScene/MovieRenderPipeline/Source/MovieRenderPipelineCore/Private/Graph/MovieGraphDefaultRenderer.cpp
@@ -134,6 +134,7 @@ void UMovieGraphDefaultRenderer::SetupRenderingPipelineForShot(UMoviePipelineExe
 		
 		FMovieGraphRenderPassSetupData SetupData;
 		SetupData.Renderer = this;
+		SetupData.EvaluatedConfig = EvaluatedConfig;
 		
 		for (const TTuple<FMovieGraphRenderDataIdentifier, TArray<TWeakObjectPtr<UMovieGraphRenderPassNode>>>& BranchRenderer : Pass.BranchRenderers)
 		{
diff --git a/Engine/Plugins/MovieScene/MovieRenderPipeline/Source/MovieRenderPipelineCore/Public/Graph/MovieGraphDataTypes.h b/Engine/Plugins/MovieScene/MovieRenderPipeline/Source/MovieRenderPipelineCore/Public/Graph/MovieGraphDataTypes.h
--- a/Engine/Plugins/MovieScene/MovieRenderPipeline/Source/MovieRenderPipelineCore/Public/Graph/MovieGraphDataTypes.h
+++ b/Engine/Plugins/MovieScene/MovieRenderPipeline/Source/MovieRenderPipelineCore/Public/Graph/MovieGraphDataTypes.h
@@ -348,6 +348,7 @@ struct FMovieGraphRenderPassSetupData
 {
 	TWeakObjectPtr<class UMovieGraphDefaultRenderer> Renderer;
 	TArray<FMovieGraphRenderPassLayerData> Layers;
+	TObjectPtr<UMovieGraphEvaluatedConfig> EvaluatedConfig;
 };
 
 
diff --git a/Engine/Plugins/MovieScene/MovieRenderPipeline/Source/MovieRenderPipelineCore/Public/Graph/Nodes/MovieGraphClassPropertyModifier.h b/Engine/Plugins/MovieScene/MovieRenderPipeline/Source/MovieRenderPipelineCore/Public/Graph/Nodes/MovieGraphClassPropertyModifier.h
--- a/Engine/Plugins/MovieScene/MovieRenderPipeline/Source/MovieRenderPipelineCore/Public/Graph/Nodes/MovieGraphClassPropertyModifier.h
+++ b/Engine/Plugins/MovieScene/MovieRenderPipeline/Source/MovieRenderPipelineCore/Public/Graph/Nodes/MovieGraphClassPropertyModifier.h
@@ -18,38 +18,38 @@ struct FMovieGraphPropertyReference
 	GENERATED_BODY()
 
 public:
-	FMovieGraphPropertyReference() = default;
+	UE_API FMovieGraphPropertyReference() = default;
 
 	/**
 	 * Providing the FProperty is optional. Providing it just bypasses the lookup of the FProperty that will occur if only the property name and
 	 * class/component are provided.
 	 */
-	FMovieGraphPropertyReference(
+	UE_API FMovieGraphPropertyReference(
 		const TObjectPtr<UClass>& InPropertyClass,
 		const TSubclassOf<UActorComponent>& InComponentClass,
 		const FName& InPropertyName,
 		FProperty* InProperty = nullptr);
 
 	/** Gets the value type of the referenced property. */
-	EMovieGraphValueType GetValueType() const;
+	UE_API EMovieGraphValueType GetValueType() const;
 
 	/** Gets the value type object of the referenced property (for objects, structs, and enums). */
-	UObject* GetValueTypeObject() const;
+	UE_API UObject* GetValueTypeObject() const;
 
 	/** Gets the FProperty that backs this property reference. */
-	FProperty* GetProperty() const;
+	UE_API FProperty* GetProperty() const;
 
 	/** Gets the name of the property that this reference refers to. */
-	FName GetPropertyName() const;
+	UE_API FName GetPropertyName() const;
 
 	/** Gets the actor class that this property exists in. If the property exists in a component, this is the class that the component belongs to. */
-	TSubclassOf<AActor> GetPropertyClass() const;
+	UE_API TSubclassOf<AActor> GetPropertyClass() const;
 
 	/** Gets the class of the component that the property exists in. Can be nullptr if the property does not exist in a component. */
-	TSubclassOf<UActorComponent> GetComponentClass() const;
+	UE_API TSubclassOf<UActorComponent> GetComponentClass() const;
 
 	/** Determines if this property reference points to a valid property. */
-	bool IsValidProperty() const;
+	UE_API bool IsValidProperty() const;
 
 	// Called by TStructOpsTypeTraits
 	void PostSerialize(const FArchive& Ar);
@@ -113,7 +113,7 @@ class UMovieGraphPropertyModifier final : public UMovieGraphCollectionModifier
 	GENERATED_BODY()
 
 public:
-	UMovieGraphPropertyModifier();
+	UE_API UMovieGraphPropertyModifier();
 
 	//~ Begin UMovieGraphModifierBase Interface
 	UE_API virtual void ApplyModifier(const UWorld* World) override;
diff --git a/Engine/Plugins/MovieScene/MovieRenderPipeline/Source/MovieRenderPipelineRenderPasses/Private/Graph/Nodes/MovieGraphDeferredPassNode.cpp b/Engine/Plugins/MovieScene/MovieRenderPipeline/Source/MovieRenderPipelineRenderPasses/Private/Graph/Nodes/MovieGraphDeferredPassNode.cpp
--- a/Engine/Plugins/MovieScene/MovieRenderPipeline/Source/MovieRenderPipelineRenderPasses/Private/Graph/Nodes/MovieGraphDeferredPassNode.cpp
+++ b/Engine/Plugins/MovieScene/MovieRenderPipeline/Source/MovieRenderPipelineRenderPasses/Private/Graph/Nodes/MovieGraphDeferredPassNode.cpp
@@ -2,12 +2,101 @@
 
 #include "Graph/Nodes/MovieGraphDeferredPassNode.h"
 
+#include "Graph/MovieGraphDataTypes.h"
 #include "Graph/Renderers/MovieGraphDeferredPass.h"
 #include "MoviePipelineTelemetry.h"
 #include "MoviePipelineUtils.h"
+#include "MovieRenderPipelineCoreModule.h"
 
 #include UE_INLINE_GENERATED_CPP_BY_NAME(MovieGraphDeferredPassNode)
 
+FCriticalSection UMovieGraphDeferredRenderPassNode::PassInstanceFactoriesCS;
+TArray<UMovieGraphDeferredRenderPassNode::FNamedPassFactory> UMovieGraphDeferredRenderPassNode::PassInstanceFactories;
+
+void UMovieGraphDeferredRenderPassNode::RegisterPassInstanceFactory(FName InName, FPassInstanceFactory InFactory)
+{
+	FScopeLock Lock(&PassInstanceFactoriesCS);
+
+	// Replace existing factory with the same name, or add new
+	for (FNamedPassFactory& Existing : PassInstanceFactories)
+	{
+		if (Existing.Name == InName)
+		{
+			Existing.Factory = MoveTemp(InFactory);
+			return;
+		}
+	}
+	PassInstanceFactories.Add({ InName, MoveTemp(InFactory) });
+}
+
+void UMovieGraphDeferredRenderPassNode::UnregisterPassInstanceFactory(FName InName)
+{
+	FScopeLock Lock(&PassInstanceFactoriesCS);
+
+	PassInstanceFactories.RemoveAll([InName](const FNamedPassFactory& Entry)
+	{
+		return Entry.Name == InName;
+	});
+}
+
+void UMovieGraphDeferredRenderPassNode::SetupImpl(const FMovieGraphRenderPassSetupData& InSetupData)
+{
+	for (const FMovieGraphRenderPassLayerData& LayerData : InSetupData.Layers)
+	{
+		TUniquePtr<UE::MovieGraph::Rendering::FMovieGraphImagePassBase> RendererInstance;
+		FName WinningFactory;
+
+		// Query registered factories — first non-null result wins
+		{
+			FScopeLock Lock(&PassInstanceFactoriesCS);
+
+			for (const FNamedPassFactory& Entry : PassInstanceFactories)
+			{
+				if (!Entry.Factory)
+				{
+					continue;
+				}
+
+				FPassInstanceFactoryContext Context{LayerData, InSetupData.Renderer.Get(), InSetupData.EvaluatedConfig};
+				TUniquePtr<UE::MovieGraph::Rendering::FMovieGraphImagePassBase> Override =
+					Entry.Factory(Context);
+
+				if (!Override)
+				{
+					continue;
+				}
+
+				if (RendererInstance)
+				{
+					// Conflict: multiple factories want to override this layer
+					UE_LOG(LogMovieRenderPipeline, Warning,
+						TEXT("Deferred pass factory conflict for layer '%s': "
+							 "'%s' was already selected, ignoring '%s'. "
+							 "Multiple plugins are trying to override the deferred renderer for the same layer."),
+						*LayerData.LayerName, *WinningFactory.ToString(), *Entry.Name.ToString());
+				}
+				else
+				{
+					RendererInstance = MoveTemp(Override);
+					WinningFactory = Entry.Name;
+				}
+			}
+		}
+
+		// Fall back to default deferred pass
+		if (!RendererInstance)
+		{
+			RendererInstance = CreateInstance();
+		}
+
+		if (RendererInstance)
+		{
+			RendererInstance->Setup(InSetupData.Renderer, this, LayerData);
+			CurrentInstances.Add(MoveTemp(RendererInstance));
+		}
+	}
+}
+
 TUniquePtr<UE::MovieGraph::Rendering::FMovieGraphImagePassBase> UMovieGraphDeferredRenderPassNode::CreateInstance() const
 {
 	return MakeUnique<UE::MovieGraph::Rendering::FMovieGraphDeferredPass>();
diff --git a/Engine/Plugins/MovieScene/MovieRenderPipeline/Source/MovieRenderPipelineRenderPasses/Private/Graph/Renderers/MovieGraphDeferredPass.cpp b/Engine/Plugins/MovieScene/MovieRenderPipeline/Source/MovieRenderPipelineRenderPasses/Private/Graph/Renderers/MovieGraphDeferredPass.cpp
--- a/Engine/Plugins/MovieScene/MovieRenderPipeline/Source/MovieRenderPipelineRenderPasses/Private/Graph/Renderers/MovieGraphDeferredPass.cpp
+++ b/Engine/Plugins/MovieScene/MovieRenderPipeline/Source/MovieRenderPipelineRenderPasses/Private/Graph/Renderers/MovieGraphDeferredPass.cpp
@@ -52,6 +52,71 @@ FSceneViewStateInterface* FMovieGraphDeferredPass::GetSceneViewState(UMovieGraph
 	return SceneViewStates[TileIndex].GetReference();
 }
 
+FMovieGraphDeferredPass::FResolutionAndCameraInfo FMovieGraphDeferredPass::GetResolutionAndCameraInfo(UMovieGraphEvaluatedConfig* InEvaluatedConfig) const
+{
+	FResolutionAndCameraInfo Info;
+
+	const UMovieGraphImagePassBaseNode* ParentNodeThisFrame = GetParentNode(InEvaluatedConfig);
+	if (!ParentNodeThisFrame)
+	{
+		// GetParentNode() will log a warning; no need to log another warning here.
+		return Info;
+	}
+
+	constexpr bool bIncludeCDOs = false;
+	const UMovieGraphCameraSettingNode* CameraNode = InEvaluatedConfig->GetSettingForBranch<UMovieGraphCameraSettingNode>(LayerData.BranchName, bIncludeCDOs);
+
+	Info.CameraInfo = GetRenderer()->GetCameraInfo(LayerData.CameraIndex);
+
+	// Get the overscan fraction and if all cameras should be rendered
+	{
+		// By default, use the overscan fraction specified directly on the camera actor
+		Info.OverscanFraction = GetRenderer()->GetCameraOverscan(LayerData.CameraIndex);
+
+		// The Camera Settings node in the graph may specify an overscan override
+		if (CameraNode)
+		{
+			Info.bRenderAllCameras = CameraNode->bRenderAllCameras;
+		
+			if (CameraNode->bOverride_OverscanPercentage)
+			{
+				Info.OverscanFraction = FMath::Clamp(CameraNode->OverscanPercentage / 100.f, 0.f, 1.f);
+			}
+		}
+	}
+
+	Info.CameraAspectRatio = (Info.CameraInfo.bAllowCameraAspectRatio && Info.CameraInfo.ViewInfo.bConstrainAspectRatio)
+		? Info.CameraInfo.ViewInfo.AspectRatio
+		: 0.f;
+
+	// The Accumulator Resolution is the size of the accumulator we should allocate. This can be bigger than the final output resolution due to camera overscan and
+	// later cropping. Eg, 1920x1080 w/ 10% overscan makes a 2112x1188 accumulator, and then the overscan crop rectangle crops it back down to 1080 for jpeg/etc,
+	// and for exr it sticks it in the margins.
+	Info.AccumulatorResolution = UMovieGraphBlueprintLibrary::GetOverscannedResolution(InEvaluatedConfig, Info.OverscanFraction, Info.CameraAspectRatio);
+
+	// Specifies a crop within the AccumulatorResolution for us to take center-out crops later where needed.
+	Info.AccumulatorResolutionCropRect = UMovieGraphBlueprintLibrary::GetOverscanCropRectangle(InEvaluatedConfig, Info.OverscanFraction, Info.CameraAspectRatio);
+
+	// This is what the actual renders would be at, without overlap.
+	const FIntPoint TileCount = ParentNodeThisFrame->GetTileCount();
+	Info.BackbufferResolution = FIntPoint(
+		FMath::CeilToInt(static_cast<float>(Info.AccumulatorResolution.X) / static_cast<float>(TileCount.X)),
+		FMath::CeilToInt(static_cast<float>(Info.AccumulatorResolution.Y) / static_cast<float>(TileCount.Y))
+	);
+
+	// Tile size is the size of a tile before any overlap padding.
+	Info.TileSize = Info.BackbufferResolution;
+
+	// User deals with overlap as 0-100, we need it as 0-1
+	const float TileOverlapFraction = ParentNodeThisFrame->GetTileOverlapPercentage() / 100.f;
+
+	// We now apply the overlap ratio (which is effectively overscan) to the backbuffer, but we use twice the user-specified overscan. This is because
+	// in High Res Tiling it's "overlap" percentage, so 10% overlap implies 10% on each side, while for Overscan, 10% overscan implies 10% bigger.
+	Info.BackbufferResolution = UE::MoviePipeline::ScaleResolutionByOverscan(TileOverlapFraction * 2, Info.BackbufferResolution);
+
+	return Info;
+}
+
 void FMovieGraphDeferredPass::Setup(TWeakObjectPtr<UMovieGraphDefaultRenderer> InRenderer, TWeakObjectPtr<UMovieGraphImagePassBaseNode> InRenderPassNode, const FMovieGraphRenderPassLayerData& InLayer)
 {
 	FMovieGraphImagePassBase::Setup(InRenderer, InRenderPassNode, InLayer);
@@ -245,11 +310,82 @@ UMovieGraphImagePassBaseNode* FMovieGraphDeferredPass::GetParentNode(UMovieGraph
 	return ParentNode;
 }
 
+FMovieGraphDeferredPass::FCreateViewResult::FCreateViewResult(const TSharedRef<FSceneViewFamilyContext> InViewFamily)
+	: ViewFamily(InViewFamily)
+{
+}
+
 bool FMovieGraphDeferredPass::HasRenderResourceParametersChanged(const FIntPoint& AccumulatorResolution, const FIntPoint& BackbufferResolution) const
 {
 	return (PrevAccumulatorResolution != AccumulatorResolution) || (PrevBackbufferResolution != BackbufferResolution);
 }
 
+FMovieGraphDeferredPass::FCreateViewResult FMovieGraphDeferredPass::CreateConfiguredView(
+	UE::MovieGraph::DefaultRenderer::FCameraInfo& InOutCameraInfo,
+	const FViewFamilyInitData& InViewFamilyInitData,
+	const FIntPoint& InBackbufferResolution,
+	const FIntPoint& InAccumulatorResolution,
+	UMovieGraphImagePassBaseNode* InParentNode,
+	const int32 TileX,
+	const int32 TileY)
+{
+	FSceneViewStateInterface* ViewState = GetSceneViewState(InParentNode, TileX, TileY);
+
+	// Optional system memory mirroring
+	if (ViewState && InParentNode->GetEnableHistoryPerTile() && InParentNode->GetEnablePageToSystemMemory())
+	{
+		ViewState->SystemMemoryMirrorRestore(SystemMemoryMirror.Get());
+	}
+
+	const TSharedRef<FSceneViewFamilyContext> ViewFamily = CreateSceneViewFamily(InViewFamilyInitData);
+
+	FSceneViewInitOptions SceneViewInitOptions = CreateViewInitOptions(InOutCameraInfo, ViewFamily.ToSharedPtr().Get(), ViewState);
+
+	CalculateProjectionMatrix(InOutCameraInfo, SceneViewInitOptions, InBackbufferResolution, InAccumulatorResolution);
+
+	// Modify the perspective matrix to do an off center projection, with overlap for high-res tiling
+	const bool bOrthographic = InOutCameraInfo.ViewInfo.ProjectionMode == ECameraProjectionMode::Type::Orthographic;
+	ModifyProjectionMatrixForTiling(InOutCameraInfo.TilingParams, bOrthographic, SceneViewInitOptions.ProjectionMatrix, InOutCameraInfo.DoFSensorScale);
+
+	// Scale the DoF sensor scale to counteract overscan, otherwise the size of Bokeh changes when you have Overscan enabled.
+	// This needs to come after we modify it for Tiling.
+	InOutCameraInfo.DoFSensorScale *= 1.0 + InOutCameraInfo.ViewInfo.GetOverscan();
+
+	// Construct a View to go within this family.
+	FSceneView* NewView = CreateSceneView(SceneViewInitOptions, ViewFamily, InOutCameraInfo);
+
+	// Object Occlusion/Histories — force reset when tiling to avoid queries from other tiles
+	{
+		const FIntPoint TileCount = InOutCameraInfo.TilingParams.TileCount;
+		if (TileCount.X > 1 || TileCount.Y > 1)
+		{
+			NewView->bForceCameraVisibilityReset = true;
+		}
+	}
+
+	// Viewport-look mode may need to apply additional customizations to the view
+#if WITH_EDITOR
+	constexpr bool bIncludeCDOs = false;
+	constexpr bool bExactMatch = true;
+	if (const UMovieGraphApplyViewportLookNode* ViewportLookNode = InViewFamilyInitData.TimeData.EvaluatedConfig->GetSettingForBranch<UMovieGraphApplyViewportLookNode>(UMovieGraphNode::GlobalsPinName, bIncludeCDOs, bExactMatch))
+	{
+		ViewportLookNode->UpdateSceneView(NewView);
+	}
+#endif	// WITH_EDITOR
+
+	// Then apply Movie Render Queue specific overrides to the ViewFamily, and then to the SceneView.
+	ApplyMovieGraphOverridesToViewFamily(ViewFamily, InViewFamilyInitData);
+
+	ApplyMovieGraphOverridesToSceneView(ViewFamily, InViewFamilyInitData, InOutCameraInfo);
+
+	FCreateViewResult Result(ViewFamily);
+	Result.View = NewView;
+	Result.ViewFamilyInitData = InViewFamilyInitData;
+	Result.SceneViewInitOptions = SceneViewInitOptions;
+
+	return Result;
+}
+
 void FMovieGraphDeferredPass::Render(const FMovieGraphTraversalContext& InFrameTraversalContext, const FMovieGraphTimeStepData& InTimeData)
 {
 	// ToDo: InFrameTraversalContext includes a copy of TimeData, but may be the one cached at the first temporal sample,
@@ -274,46 +410,9 @@ void FMovieGraphDeferredPass::Render(const FMovieGraphTraversalContext& InFrameT
 	const ESceneCaptureSource SceneCaptureSource = ParentNodeThisFrame->GetDisableToneCurve() ? ESceneCaptureSource::SCS_FinalColorHDR : ESceneCaptureSource::SCS_FinalToneCurveHDR;
 	const EAntiAliasingMethod AntiAliasingMethod = UE::MovieRenderPipeline::GetEffectiveAntiAliasingMethod(ParentNodeThisFrame->GetOverrideAntiAliasing(), ParentNodeThisFrame->GetAntiAliasingMethod());
 
-	UE::MovieGraph::DefaultRenderer::FCameraInfo GlobalCameraInfo = GetRenderer()->GetCameraInfo(LayerData.CameraIndex);
-	float CameraOverscanFraction = GetRenderer()->GetCameraOverscan(LayerData.CameraIndex);
-
-	// We need the aspect ratio from the camera, but below it ends up getting modified and we pass references around
-	// so they need to stay unique instances below.
-	const float CameraAspectRatio = GlobalCameraInfo.bAllowCameraAspectRatio && GlobalCameraInfo.ViewInfo.bConstrainAspectRatio ? GlobalCameraInfo.ViewInfo.AspectRatio : 0.f;
-
-	// Camera nodes are optional
-	const bool bIncludeCDOs = false;
-	bool bRenderAllCameras = false;
-	const UMovieGraphCameraSettingNode* CameraNode = InTimeData.EvaluatedConfig->GetSettingForBranch<UMovieGraphCameraSettingNode>(LayerData.BranchName, bIncludeCDOs);
-	if (CameraNode)
-	{
-		bRenderAllCameras = CameraNode->bRenderAllCameras;
-		if (CameraNode->bOverride_OverscanPercentage)
-		{
-			CameraOverscanFraction = FMath::Clamp(CameraNode->OverscanPercentage / 100.f, 0.f, 1.f);
-		}
-	}
-
-	// The Accumulator Resolution is the size of the accumulator we should allocate. This can be bigger than the final output resolution due to camera overscan and
-	// later cropping, ie: a 1920x1080 /w 10% overscan makes a 2112x1188 accumulator, and then the overscan crop rectangle crops it back down to 1080 for jpeg/etc,
-	// and for exr it sticks it in the margins.
-	const FIntPoint AccumulatorResolution = UMovieGraphBlueprintLibrary::GetOverscannedResolution(InTimeData.EvaluatedConfig, CameraOverscanFraction, CameraAspectRatio);
-	// Specifies a crop within the AccumulatorResolution for us to take center-out crops later where needed.
-	const FIntRect AccumulatorResolutionCropRect = UMovieGraphBlueprintLibrary::GetOverscanCropRectangle(InTimeData.EvaluatedConfig, CameraOverscanFraction, CameraAspectRatio);
-	
-	// This is what the actual renders would be at, without overlap.
-	FIntPoint BackbufferResolution = FIntPoint(
-		FMath::CeilToInt((float)AccumulatorResolution.X / (float)TileCount.X),
-		FMath::CeilToInt((float)AccumulatorResolution.Y / (float)TileCount.Y)
-	);
-
-	// Tile size is the size of a tile before any overlap padding.
-	const FIntPoint TileSize = BackbufferResolution;
-
-	// We now apply the overlap ratio (which is effectively overscan) to the backbuffer, but we use twice the user-specified overscan. This is because
-	// in High Res Tiling it's "overlap" percentage, so 10% overlap implies 10% on each side, while for Overscan, 10% overscan implies 10% bigger.
-	BackbufferResolution = UE::MoviePipeline::ScaleResolutionByOverscan(TileOverlapFraction * 2, BackbufferResolution);
-
+	const FResolutionAndCameraInfo ResolutionAndCameraInfo = GetResolutionAndCameraInfo(InTimeData.EvaluatedConfig.Get());
+	const FIntPoint AccumulatorResolution = ResolutionAndCameraInfo.AccumulatorResolution;
+	const FIntPoint BackbufferResolution = ResolutionAndCameraInfo.BackbufferResolution;
 	if (HasRenderResourceParametersChanged(AccumulatorResolution, BackbufferResolution))
 	{
 		UE_LOG(LogMovieRenderPipeline, Log, TEXT("Set-up Renderer: %s Layer: %s AccumulatorResolution: [%d, %d] BackbufferResolution: [%d, %d]"),
@@ -325,7 +424,6 @@ void FMovieGraphDeferredPass::Render(const FMovieGraphTraversalContext& InFrameT
 
 	DefaultRenderer::FRenderTargetInitParams RenderTargetInitParams = GetRenderTargetInitParams(InTimeData, BackbufferResolution);
 
-
 	for (int32 TileY = 0; TileY < TileCount.Y; TileY++)
 	{
 		for (int32 TileX = 0; TileX < TileCount.X; TileX++)
@@ -358,6 +456,7 @@ void FMovieGraphDeferredPass::Render(const FMovieGraphTraversalContext& InFrameT
 				// These are the parameters of our camera 
 				UE::MovieGraph::DefaultRenderer::FCameraInfo CameraInfo = GetRenderer()->GetCameraInfo(LayerData.CameraIndex);
 
+				FIntPoint TileSize = ResolutionAndCameraInfo.TileSize;
 				FIntPoint OverlappedPad = FIntPoint(FMath::CeilToInt(TileSize.X * TileOverlapFraction), FMath::CeilToInt(TileSize.Y * TileOverlapFraction));
 
 				CameraInfo.bAllowCameraAspectRatio = true;
@@ -370,14 +469,16 @@ void FMovieGraphDeferredPass::Render(const FMovieGraphTraversalContext& InFrameT
 				CameraInfo.SamplingParams.SpatialSampleIndex = SpatialIndex;
 				CameraInfo.SamplingParams.SpatialSampleCount = NumSpatialSamples;
 				CameraInfo.SamplingParams.SeedOffset = SeedOffset;
-				CameraInfo.ProjectionMatrixJitterAmount = FVector2D((SpatialShiftAmount.X) * 2.0f / (float)BackbufferResolution.X, SpatialShiftAmount.Y * -2.0f / (float)BackbufferResolution.Y);
-				CameraInfo.bUseCameraManagerPostProcess = !bRenderAllCameras;
+				FVector2D DefaultJitter = FVector2D(
+					SpatialShiftAmount.X * 2.0f / static_cast<float>(BackbufferResolution.X),
+					SpatialShiftAmount.Y * -2.0f / static_cast<float>(BackbufferResolution.Y));
+				CameraInfo.ProjectionMatrixJitterAmount = ComputeProjectionMatrixJitter(DefaultJitter, BackbufferResolution);
+				CameraInfo.bUseCameraManagerPostProcess = !ResolutionAndCameraInfo.bRenderAllCameras;
 
 				// Make sure our View Info is updated with the correct values from global camera overscan, this needs to happen
 				// before we start tinkering with the matrices for tiling.
 				CameraInfo.ViewInfo.ClearOverscan();
-				CameraInfo.ViewInfo.ApplyOverscan(CameraOverscanFraction);
-
+				CameraInfo.ViewInfo.ApplyOverscan(ResolutionAndCameraInfo.OverscanFraction);
 
 				// For this particular tile, what is the offset into the output image
 				FIntPoint OverlappedOffset = FIntPoint(CameraInfo.TilingParams.TileIndexes.X * TileSize.X - OverlappedPad.X, CameraInfo.TilingParams.TileIndexes.Y * TileSize.Y - OverlappedPad.Y);
@@ -401,58 +502,9 @@ void FMovieGraphDeferredPass::Render(const FMovieGraphTraversalContext& InFrameT
 				ViewFamilyInitData.ViewModeIndex = ParentNodeThisFrame->GetViewModeIndex();
 				ViewFamilyInitData.ProjectionMode = CameraInfo.ViewInfo.ProjectionMode;
 
-				FSceneViewStateInterface* ViewState = GetSceneViewState(ParentNodeThisFrame, TileX, TileY);
-				
-				// Optional system memory mirroring
-				if (ViewState && ParentNodeThisFrame->GetEnableHistoryPerTile() && ParentNodeThisFrame->GetEnablePageToSystemMemory())
-				{
-					ViewState->SystemMemoryMirrorRestore(SystemMemoryMirror.Get());
-				}
-
-				TSharedRef<FSceneViewFamilyContext> ViewFamily = CreateSceneViewFamily(ViewFamilyInitData);
-
-				FSceneViewInitOptions SceneViewInitOptions = CreateViewInitOptions(CameraInfo, ViewFamily.ToSharedPtr().Get(), ViewState);
-		
-		
-				CalculateProjectionMatrix(CameraInfo, SceneViewInitOptions, BackbufferResolution, AccumulatorResolution);
-	
-				// Modify the perspective matrix to do an off center projection, with overlap for high-res tiling
-				const bool bOrthographic = CameraInfo.ViewInfo.ProjectionMode == ECameraProjectionMode::Type::Orthographic;
-				ModifyProjectionMatrixForTiling(CameraInfo.TilingParams, bOrthographic, SceneViewInitOptions.ProjectionMatrix, CameraInfo.DoFSensorScale);
-		
-				// Scale the DoF sensor scale to counteract overscan, otherwise the size of Bokeh changes when you have Overscan enabled.
-				// This needs to come after we modify it for Tiling.
-				CameraInfo.DoFSensorScale *= 1.0 + CameraInfo.ViewInfo.GetOverscan();
-		 
-				// Construct a View to go within this family.
-				FSceneView* NewView = CreateSceneView(SceneViewInitOptions, ViewFamily, CameraInfo);
-
-				// Object Occlusion/Histories
-				{
-					// If we're using tiling, we force the reset of histories each frame so that we don't use the previous tile's
-					// object occlusion queries, as that causes things to disappear from some views.
-					if (TileCount.X > 1 || TileCount.Y > 1)
-					{
-						NewView->bForceCameraVisibilityReset = true;
-					}
-				}
-
-				// Viewport-look mode may need to apply additional customizations to the view
-		#if WITH_EDITOR
-				constexpr bool bExactMatch = true;
-				if (const UMovieGraphApplyViewportLookNode* ViewportLookNode = InTimeData.EvaluatedConfig->GetSettingForBranch<UMovieGraphApplyViewportLookNode>(UMovieGraphNode::GlobalsPinName, bIncludeCDOs, bExactMatch))
-				{
-					ViewportLookNode->UpdateSceneView(NewView);
-				}
-		#endif	// WITH_EDITOR
-		
-				// Then apply Movie Render Queue specific overrides to the ViewFamily, and then to the SceneView.
-				ApplyMovieGraphOverridesToViewFamily(ViewFamily, ViewFamilyInitData);
-		
-				// ToDo: This really only needs access to the ViewFamily for path tracer related things,
-				// and would rather just take a FSceneView* 
-				ApplyMovieGraphOverridesToSceneView(ViewFamily, ViewFamilyInitData, CameraInfo);
-		
+				FCreateViewResult ViewResult = CreateConfiguredView(CameraInfo, ViewFamilyInitData, BackbufferResolution, AccumulatorResolution, ParentNodeThisFrame, TileX, TileY);
+				TSharedRef<FSceneViewFamilyContext> ViewFamily = ViewResult.ViewFamily;
+				FSceneView* NewView = ViewResult.View;
 
 				FHitProxyConsumer* HitProxyConsumer = nullptr;
 				const float DPIScale = 1.0f;
@@ -480,7 +532,7 @@ void FMovieGraphDeferredPass::Render(const FMovieGraphTraversalContext& InFrameT
 					SampleState.OverlappedOffset = OverlappedOffset;
 					SampleState.OverlappedSubpixelShift = OverlappedSubpixelShift;
 					SampleState.OverscanFraction = CameraInfo.ViewInfo.GetOverscan();
-					SampleState.CropRectangle = AccumulatorResolutionCropRect;
+					SampleState.CropRectangle = ResolutionAndCameraInfo.AccumulatorResolutionCropRect;
 					SampleState.bAllowOCIO = ParentNodeThisFrame->GetAllowOCIO();
 					SampleState.bForceLosslessCompression = ParentNodeThisFrame->GetForceLosslessCompression();
 					SampleState.bAllowsCompositing = ParentNodeThisFrame->GetAllowsCompositing();
@@ -492,7 +544,7 @@ void FMovieGraphDeferredPass::Render(const FMovieGraphTraversalContext& InFrameT
 
 				ApplyMovieGraphOverridesToSampleState(SampleState);
 
-				MetadataHelper::UpdateSpatialSampleMetadata(SceneViewInitOptions, SampleState.AdditionalFileMetadata, RenderDataIdentifier);
+				MetadataHelper::UpdateSpatialSampleMetadata(ViewResult.SceneViewInitOptions, SampleState.AdditionalFileMetadata, RenderDataIdentifier);
 				FString MetadataPrefix = UE::MoviePipeline::GetMetadataPrefixPath(RenderDataIdentifier);
 					UE::MoviePipeline::GetMetadataFromCameraLocRot(
 					CameraInfo.ViewInfo.Location, CameraInfo.ViewInfo.Rotation, 
@@ -582,9 +634,12 @@ void FMovieGraphDeferredPass::Render(const FMovieGraphTraversalContext& InFrameT
 					});
 
 				// After submission, if we're paging to system memory, mark the resources for download into system memory.
-				if (ViewState && ParentNodeThisFrame->GetEnableHistoryPerTile() && ParentNodeThisFrame->GetEnablePageToSystemMemory())
+				if (ParentNodeThisFrame->GetEnableHistoryPerTile() && ParentNodeThisFrame->GetEnablePageToSystemMemory())
 				{
-					ViewState->SystemMemoryMirrorBackup(SystemMemoryMirror.Get());
+					if (FSceneViewStateInterface* ViewState = GetSceneViewState(ParentNodeThisFrame, TileX, TileY))
+					{
+						ViewState->SystemMemoryMirrorBackup(SystemMemoryMirror.Get());
+					}
 				}
 
 				// Data may not be something we actually want to read back and write to disk after rendering.
@@ -737,4 +792,10 @@ void FMovieGraphDeferredPass::PostRendererSubmission(
 	
 	FMovieGraphImagePassBase::PostRendererSubmission(InSampleState, InRenderTargetInitParams, InCanvas, InCameraInfo);
 }
+
+FVector2D FMovieGraphDeferredPass::ComputeProjectionMatrixJitter(const FVector2D& InDefaultJitter, const FIntPoint& InBackbufferResolution) const
+{
+	return InDefaultJitter;
+}
+
 } // namespace UE::MovieGraph::Rendering
diff --git a/Engine/Plugins/MovieScene/MovieRenderPipeline/Source/MovieRenderPipelineRenderPasses/Public/Graph/Nodes/MovieGraphDeferredPassNode.h b/Engine/Plugins/MovieScene/MovieRenderPipeline/Source/MovieRenderPipelineRenderPasses/Public/Graph/Nodes/MovieGraphDeferredPassNode.h
--- a/Engine/Plugins/MovieScene/MovieRenderPipeline/Source/MovieRenderPipelineRenderPasses/Public/Graph/Nodes/MovieGraphDeferredPassNode.h
+++ b/Engine/Plugins/MovieScene/MovieRenderPipeline/Source/MovieRenderPipelineRenderPasses/Public/Graph/Nodes/MovieGraphDeferredPassNode.h
@@ -1,10 +1,13 @@
 ﻿// Copyright Epic Games, Inc. All Rights Reserved.
 
 #pragma once
+
 #include "Graph/Nodes/MovieGraphImagePassBaseNode.h"
 #include "MoviePipelineDeferredPasses.h"
 #include "MovieGraphDeferredPassNode.generated.h"
 
+class UMovieGraphEvaluatedConfig;
+
 #define UE_API MOVIERENDERPIPELINERENDERPASSES_API
 
 /**
@@ -22,6 +25,38 @@ public:
 	UE_API virtual void UpdateTelemetry(FMoviePipelineShotRenderTelemetry* InTelemetry) const override;
 	UE_API virtual void ResolveTokenContainingProperties(TFunction<void(FString&)>& ResolveFunc, const FMovieGraphTokenResolveContext& InContext) override;
 
+	/** Context provided to pass instance factories during setup. */
+	struct FPassInstanceFactoryContext
+	{
+		const FMovieGraphRenderPassLayerData& LayerData;
+		UMovieGraphDefaultRenderer* Renderer = nullptr;
+		UMovieGraphEvaluatedConfig* EvaluatedConfig = nullptr;
+	};
+
+	/**
+	 * Factory function type for plugins that want to override deferred pass instance creation.
+	 * Called once per layer during setup. Return a valid instance to override the default
+	 * FMovieGraphDeferredPass, or nullptr to decline (fall back to next factory or default).
+	 *
+	 * Note: If multiple plugins register factories that claim the same layer, a warning
+	 * is logged. The first registered factory that returns a non-null pass wins.
+	 */
+	using FPassInstanceFactory = TFunction<TUniquePtr<UE::MovieGraph::Rendering::FMovieGraphImagePassBase>(
+		const FPassInstanceFactoryContext& /*InContext*/)>;
+
+	/**
+	 * Register a named factory that can override pass instance creation. This is used to tweak the
+	 * behavior of the Deferred Renderer node. Use sparingly; generally, creating a new node is the
+	 * best approach.
+	 * @param InName     Unique name identifying the pass instance factory.
+	 *                   If a factory with this name is already registered, it is replaced.
+	 * @param InFactory  The factory function.
+	 */
+	static UE_API void RegisterPassInstanceFactory(FName InName, FPassInstanceFactory InFactory);
+
+	/** Unregister a previously registered factory by name. */
+	static UE_API void UnregisterPassInstanceFactory(FName InName);
+
 #if WITH_EDITOR
 	UE_API virtual FText GetNodeTitle(const bool bGetDescriptive = false) const override;
 	UE_API virtual FSlateIcon GetIconAndTint(FLinearColor& OutColor) const override;
@@ -29,6 +64,7 @@ public:
 
 protected:
 	// UMovieGraphRenderPassNode Interface
+	UE_API virtual void SetupImpl(const FMovieGraphRenderPassSetupData& InSetupData) override;
 	UE_API virtual TUniquePtr<UE::MovieGraph::Rendering::FMovieGraphImagePassBase> CreateInstance() const override;
 	UE_API virtual bool GetWriteAllSamples() const override;
 	UE_API virtual TArray<FMoviePipelinePostProcessPass> GetAdditionalPostProcessMaterials() const override;
@@ -215,6 +251,23 @@ public:
 	*/
 	UPROPERTY(EditAnywhere, BlueprintReadWrite, Category = "High Resolution Tiling", meta = (EditCondition = "bOverride_bPageToSystemMemory"))
 	bool bPageToSystemMemory;
+
+private:
+	/** Information identifying a registered pass instance factory. */
+	struct FNamedPassFactory
+	{
+		/** The name of the pass instance factory. */
+		FName Name;
+		
+		/** The function that will create the pass instance factory. */
+		FPassInstanceFactory Factory;
+	};
+
+	/** Protects mutations to PassInstanceFactories. */
+	static FCriticalSection PassInstanceFactoriesCS;
+
+	/** The pass instance factories that have been registered. */
+	static TArray<FNamedPassFactory> PassInstanceFactories;
 };
 
 #undef UE_API
diff --git a/Engine/Plugins/MovieScene/MovieRenderPipeline/Source/MovieRenderPipelineRenderPasses/Public/Graph/Renderers/MovieGraphDeferredPass.h b/Engine/Plugins/MovieScene/MovieRenderPipeline/Source/MovieRenderPipelineRenderPasses/Public/Graph/Renderers/MovieGraphDeferredPass.h
--- a/Engine/Plugins/MovieScene/MovieRenderPipeline/Source/MovieRenderPipelineRenderPasses/Public/Graph/Renderers/MovieGraphDeferredPass.h
+++ b/Engine/Plugins/MovieScene/MovieRenderPipeline/Source/MovieRenderPipelineRenderPasses/Public/Graph/Renderers/MovieGraphDeferredPass.h
@@ -28,12 +28,56 @@ namespace UE::MovieGraph::Rendering
 		UE_API virtual UMovieGraphImagePassBaseNode* GetParentNode(UMovieGraphEvaluatedConfig* InConfig) const override;
 		UE_API virtual bool ShouldDiscardOutput(const TSharedRef<FSceneViewFamilyContext>& InFamily, const UE::MovieGraph::DefaultRenderer::FCameraInfo& InCameraInfo) const override;
 		// End FMovieGraphImagePassBase
-			
+
 	protected:
+		/** Data about the resolution and camera that can be used in the next render. */
+		struct FResolutionAndCameraInfo
+		{
+			float OverscanFraction = 0.f;
+			float CameraAspectRatio = 0.f;
+			bool bRenderAllCameras = false;
+			DefaultRenderer::FCameraInfo CameraInfo;
+			FIntPoint BackbufferResolution;
+			FIntPoint AccumulatorResolution;
+			FIntRect AccumulatorResolutionCropRect;
+			FIntPoint TileSize;
+		};
+
+		/** Result of CreateConfiguredView() — contains the fully configured ViewFamily, View, FViewFamilyInitData, and FSceneViewInitOptions. */
+		struct FCreateViewResult
+		{
+			explicit FCreateViewResult(const TSharedRef<FSceneViewFamilyContext> InViewFamily);
+
+			TSharedRef<FSceneViewFamilyContext> ViewFamily;
+			FSceneView* View = nullptr;
+			FViewFamilyInitData ViewFamilyInitData;
+			FSceneViewInitOptions SceneViewInitOptions;
+		};
+
 		UE_API bool HasRenderResourceParametersChanged(const FIntPoint& AccumulatorResolution, const FIntPoint& BackbufferResolution) const;
 		UE_API virtual void PostRendererSubmission(const UE::MovieGraph::FMovieGraphSampleState& InSampleState, const UE::MovieGraph::DefaultRenderer::FRenderTargetInitParams& InRenderTargetInitParams, FCanvas& InCanvas, const UE::MovieGraph::DefaultRenderer::FCameraInfo& InCameraInfo) const override;
 		UE_API virtual FSceneViewStateInterface* GetSceneViewState(UMovieGraphImagePassBaseNode* ParentNodeThisFrame, int32_t TileX, int32_t TileY);
 
+		/** Provides the resolution and camera information that should be used for the next render. */
+		UE_API FResolutionAndCameraInfo GetResolutionAndCameraInfo(UMovieGraphEvaluatedConfig* InEvaluatedConfig) const;
+
+		/**
+		 * Encapsulates the shared view creation pipeline: GetSceneViewState -> CreateSceneViewFamily ->
+		 * CreateViewInitOptions -> CalculateProjectionMatrix -> ModifyProjectionMatrixForTiling ->
+		 * CreateSceneView -> ApplyMovieGraphOverrides. Used by both the base Render() and derived passes.
+		 */
+		UE_API FCreateViewResult CreateConfiguredView(
+			UE::MovieGraph::DefaultRenderer::FCameraInfo& InOutCameraInfo,
+			const FViewFamilyInitData& InViewFamilyInitData,
+			const FIntPoint& InBackbufferResolution,
+			const FIntPoint& InAccumulatorResolution,
+			UMovieGraphImagePassBaseNode* InParentNode,
+			const int32 TileX,
+			const int32 TileY);
+
+		/** Computes the projection matrix jitter amount for the current sample. Override to provide a custom jitter. */
+		UE_API virtual FVector2D ComputeProjectionMatrixJitter(const FVector2D& InDefaultJitter, const FIntPoint& InBackbufferResolution) const;
+
 	protected:
 		FMovieGraphRenderPassLayerData LayerData;
 
@@ -66,6 +110,7 @@ namespace UE::MovieGraph::Rendering
 
 		// Track the last Accumulator resolution we used, so that we can detect when it is changed and log that information.
 		FIntPoint PrevAccumulatorResolution;
+
 		// Track the last backbuffer resolution we used, so that we can detect when it is changed and log that information.
 		FIntPoint PrevBackbufferResolution;
 	}; 
diff --git a/Engine/Plugins/VirtualProduction/AccumulationDOF/Source/AccumulationDOF/AccumulationDOF.Build.cs b/Engine/Plugins/VirtualProduction/AccumulationDOF/Source/AccumulationDOF/AccumulationDOF.Build.cs
--- a/Engine/Plugins/VirtualProduction/AccumulationDOF/Source/AccumulationDOF/AccumulationDOF.Build.cs
+++ b/Engine/Plugins/VirtualProduction/AccumulationDOF/Source/AccumulationDOF/AccumulationDOF.Build.cs
@@ -20,6 +20,10 @@ public class AccumulationDOF : ModuleRules
 			{
 				"CinematicCamera",
 				"Core",
+				"Engine",
+				"MovieRenderPipelineCore",
+				"Slate",
+				"SlateCore",
 			}
 		);
 
@@ -29,7 +33,6 @@ public class AccumulationDOF : ModuleRules
 				"CoreUObject",
 				"Engine",
 				"ImageWriteQueue",
-				"MovieRenderPipelineCore",
 				"MovieRenderPipelineRenderPasses",
 				"Projects",
 				"RenderCore",
@@ -37,5 +40,10 @@ public class AccumulationDOF : ModuleRules
 				"RHI",
 			}
 		);
+
+		if (Target.bBuildEditor)
+		{
+			PrivateDependencyModuleNames.Add("PropertyEditor");
+		}
 	}
 }
diff --git a/Engine/Plugins/VirtualProduction/AccumulationDOF/Source/AccumulationDOF/Private/AccumulationDOF.cpp b/Engine/Plugins/VirtualProduction/AccumulationDOF/Source/AccumulationDOF/Private/AccumulationDOF.cpp
--- a/Engine/Plugins/VirtualProduction/AccumulationDOF/Source/AccumulationDOF/Private/AccumulationDOF.cpp
+++ b/Engine/Plugins/VirtualProduction/AccumulationDOF/Source/AccumulationDOF/Private/AccumulationDOF.cpp
@@ -1,15 +1,28 @@
 // Copyright Epic Games, Inc. All Rights Reserved.
 
 #include "AccumulationDOF.h"
+#include "AccumulationDOFComponent.h"
 #include "AccumulationDOFLog.h"
+#include "CineCameraActor.h"
+#include "Graph/MovieGraphConfig.h"
+#include "Graph/MovieGraphDefaultRenderer.h"
+#include "Graph/Nodes/MovieGraphDeferredPassNode.h"
 #include "Interfaces/IPluginManager.h"
 #include "Misc/Paths.h"
+#include "MovieGraphAccumulationDOFModifierNode.h"
+#include "MovieGraphAccumulationDOFPass.h"
 #include "ShaderCore.h"
 
+#if WITH_EDITOR
+#include "MovieGraphAccumulationDOFModifierCustomization.h"
+#include "PropertyEditorModule.h"
+#endif // WITH_EDITOR
+
 DEFINE_LOG_CATEGORY(LogAccumulationDOF);
 
 #define LOCTEXT_NAMESPACE "FAccumulationDOFModule"
 
+const FName FAccumulationDOFModule::MRGPassFactoryPassName(TEXT("AccumulationDOF"));
 
 void FAccumulationDOFModule::StartupModule()
 {
@@ -20,10 +33,80 @@ void FAccumulationDOFModule::StartupModule()
 		FString PluginShaderDirectory = FPaths::Combine(Plugin->GetBaseDir(), TEXT("Shaders"));
 		AddShaderSourceDirectoryMapping(TEXT("/Plugin/AccumulationDOF"), PluginShaderDirectory);
 	}
+
+#if WITH_EDITOR
+	FPropertyEditorModule& PropertyModule = FModuleManager::Get().LoadModuleChecked<FPropertyEditorModule>("PropertyEditor");
+	PropertyModule.RegisterCustomClassLayout("MovieGraphAccumulationDOFModifierNode",
+		FOnGetDetailCustomizationInstance::CreateStatic(&FMovieGraphAccumulationDOFModifierCustomization::MakeInstance));
+#endif	// WITH_EDITOR
+
+	// Register AccumulationDOF pass factory with the deferred renderer node.
+	// When a camera has an active UAccumulationDOFComponent, the deferred renderer will
+	// use FMovieGraphAccumulationDOFPass instead of FMovieGraphDeferredPass.
+	UMovieGraphDeferredRenderPassNode::RegisterPassInstanceFactory(
+		MRGPassFactoryPassName,
+		[](const UMovieGraphDeferredRenderPassNode::FPassInstanceFactoryContext& InContext)
+		-> TUniquePtr<UE::MovieGraph::Rendering::FMovieGraphImagePassBase>
+	{
+		if (!InContext.Renderer)
+		{
+			return nullptr;
+		}
+
+		const UE::MovieGraph::DefaultRenderer::FCameraInfo CameraInfo = InContext.Renderer->GetCameraInfo(InContext.LayerData.CameraIndex);
+		const ACineCameraActor* CineCamera = Cast<ACineCameraActor>(CameraInfo.ViewActor);
+		if (!CineCamera)
+		{
+			return nullptr;
+		}
+
+		const UAccumulationDOFComponent* DOFComponent = CineCamera->FindComponentByClass<UAccumulationDOFComponent>();
+		if (!DOFComponent)
+		{
+			return nullptr;
+		}
+
+		// First, use the component's current active state to determine if this pass will be instantiated
+		bool bDOFEnabled = DOFComponent->IsActive();
+
+		// Second, if an Accumulation DOF modifier on this layer is disabling Accumulation DOF, don't instantiate this pass.
+		// Checking the modifier is not ideal. However, the modifier is applied *after* pass instantiation, so we cannot solely
+		// rely on the component active state.
+		if (InContext.EvaluatedConfig)
+		{
+			constexpr bool bIncludeCDOs = false;
+			constexpr bool bExactMatch = true;
+			const UMovieGraphAccumulationDOFModifierNode* ModNode =
+				InContext.EvaluatedConfig->GetSettingForBranch<UMovieGraphAccumulationDOFModifierNode>(
+					InContext.LayerData.BranchName, bIncludeCDOs, bExactMatch);
+
+			if (ModNode
+				&& ModNode->bOverride_EnableAccumulationDepthOfField
+				&& ModNode->EnableAccumulationDepthOfField.Type == EMovieGraphAccumulationDOFEnableType::CustomValue)
+			{
+				bDOFEnabled = ModNode->EnableAccumulationDepthOfField.Value;
+			}
+		}
+
+		if (bDOFEnabled)
+		{
+			return MakeUnique<FMovieGraphAccumulationDOFPass>();
+		}
+
+		return nullptr;
+	});
 }
 
 void FAccumulationDOFModule::ShutdownModule()
 {
+#if WITH_EDITOR
+	if (FPropertyEditorModule* PropertyModule = FModuleManager::GetModulePtr<FPropertyEditorModule>("PropertyEditor"))
+	{
+		PropertyModule->UnregisterCustomClassLayout("MovieGraphAccumulationDOFModifierNode");
+	}
+#endif	// WITH_EDITOR
+
+	UMovieGraphDeferredRenderPassNode::UnregisterPassInstanceFactory(MRGPassFactoryPassName);
 }
 
 #undef LOCTEXT_NAMESPACE
diff --git a/Engine/Plugins/VirtualProduction/AccumulationDOF/Source/AccumulationDOF/Private/MovieGraphAccumulationDOFModifier.cpp b/Engine/Plugins/VirtualProduction/AccumulationDOF/Source/AccumulationDOF/Private/MovieGraphAccumulationDOFModifier.cpp
new file mode 100644
--- /dev/null
+++ b/Engine/Plugins/VirtualProduction/AccumulationDOF/Source/AccumulationDOF/Private/MovieGraphAccumulationDOFModifier.cpp
@@ -0,0 +1,98 @@
+﻿// Copyright Epic Games, Inc. All Rights Reserved.
+
+#include "MovieGraphAccumulationDOFModifier.h"
+
+#include "AccumulationDOFComponent.h"
+#include "Camera/CameraActor.h"
+#include "EngineUtils.h"
+
+FMovieGraphAccumulationDOFModifier_ComponentState::FMovieGraphAccumulationDOFModifier_ComponentState(UAccumulationDOFComponent* InComponent)
+	: DOFComponent(InComponent)
+{
+	if (DOFComponent.IsValid())
+	{
+		bEnableAccumulationDepthOfField = DOFComponent->IsActive();
+		NumSamples = DOFComponent->NumSamples;
+		DOFSplatSize = DOFComponent->DOFSplatSize;
+	}
+}
+
+void UMovieGraphAccumulationDOFModifier::ApplyModifier(const UWorld* World)
+{
+	OriginalComponentState.Reset();
+
+	for (TActorIterator<ACameraActor> It(World); It; ++It)
+	{
+		const ACameraActor* CameraActor = *It;
+		if (!CameraActor)
+		{
+			continue;
+		}
+
+		UAccumulationDOFComponent* AccumulationDOFComponent = CameraActor->FindComponentByClass<UAccumulationDOFComponent>();
+		if (!AccumulationDOFComponent)
+		{
+			continue;
+		}
+
+		// Found an Accumulation DOF component. First, cache its current state.
+		OriginalComponentState.Add(FMovieGraphAccumulationDOFModifier_ComponentState(AccumulationDOFComponent));
+
+		// Second, apply new state.
+		{
+			if (bOverride_EnableDOFComponents)
+			{
+				if (EnableDOFComponents.Type == EMovieGraphAccumulationDOFEnableType::CustomValue)
+				{
+					AccumulationDOFComponent->SetActive(EnableDOFComponents.Value);
+				}
+			}
+
+			if (bOverride_NumSamples)
+			{
+				AccumulationDOFComponent->NumSamples = NumSamples;
+			}
+
+			if (bOverride_DOFSplatSize)
+			{
+				AccumulationDOFComponent->DOFSplatSize = DOFSplatSize;
+			}
+		}
+
+		AccumulationDOFComponent->MarkRenderStateDirty();
+	}
+}
+
+void UMovieGraphAccumulationDOFModifier::UndoModifier()
+{
+	for (const FMovieGraphAccumulationDOFModifier_ComponentState& ComponentState : OriginalComponentState)
+	{
+		if (UAccumulationDOFComponent* CachedDOFComponent = ComponentState.DOFComponent.Get())
+		{
+			if (bOverride_EnableDOFComponents)
+			{
+				CachedDOFComponent->SetActive(ComponentState.bEnableAccumulationDepthOfField);
+			}
+
+			if (bOverride_NumSamples)
+			{
+				CachedDOFComponent->NumSamples = ComponentState.NumSamples;
+			}
+
+			if (bOverride_DOFSplatSize)
+			{
+				CachedDOFComponent->DOFSplatSize = ComponentState.DOFSplatSize;
+			}
+
+			CachedDOFComponent->MarkRenderStateDirty();
+		}
+	}
+	
+	OriginalComponentState.Reset();
+}
+
+FText UMovieGraphAccumulationDOFModifier::GetModifierName()
+{
+	static const FText ModifierName = NSLOCTEXT("MovieGraph", "AccumulationDOFModifierName", "Accumulation DOF");
+	return ModifierName;
+}
diff --git a/Engine/Plugins/VirtualProduction/AccumulationDOF/Source/AccumulationDOF/Private/MovieGraphAccumulationDOFModifier.h b/Engine/Plugins/VirtualProduction/AccumulationDOF/Source/AccumulationDOF/Private/MovieGraphAccumulationDOFModifier.h
new file mode 100644
--- /dev/null
+++ b/Engine/Plugins/VirtualProduction/AccumulationDOF/Source/AccumulationDOF/Private/MovieGraphAccumulationDOFModifier.h
@@ -0,0 +1,84 @@
+// Copyright Epic Games, Inc. All Rights Reserved.
+
+#pragma once
+
+#include "Graph/MovieGraphRenderLayerSubsystem.h"
+#include "MovieGraphAccumulationDOFModifierNode.h"
+
+#include "MovieGraphAccumulationDOFModifier.generated.h"
+
+class UAccumulationDOFComponent;
+
+#define UE_API ACCUMULATIONDOF_API
+
+/** The state of an Accumulation DOF component before modification. */
+USTRUCT()
+struct FMovieGraphAccumulationDOFModifier_ComponentState
+{
+	GENERATED_BODY()
+
+	FMovieGraphAccumulationDOFModifier_ComponentState() = default;
+	explicit FMovieGraphAccumulationDOFModifier_ComponentState(UAccumulationDOFComponent* InComponent);
+
+	/** The Accumulation DOF component that the below cached state is for. */
+	TWeakObjectPtr<UAccumulationDOFComponent> DOFComponent;
+
+	/** The original enable/disable state of the Accumulation DOF component. */
+	bool bEnableAccumulationDepthOfField = false;
+
+	/** The original number of samples set on the Accumulation DOF component. */
+	int32 NumSamples = 0;
+
+	/** The original splat size set on the Accumulation DOF component. */
+	float DOFSplatSize = 0.f;
+};
+
+/**
+ * Updates all Accumulation DOF components within the world to the settings specified on this modifier.
+ */
+UCLASS(MinimalAPI, BlueprintType)
+class UMovieGraphAccumulationDOFModifier final : public UMovieGraphModifierBase
+{
+	GENERATED_BODY()
+
+public:
+	UMovieGraphAccumulationDOFModifier() = default;
+
+	//~ Begin UMovieGraphModifierBase Interface
+	UE_API virtual void ApplyModifier(const UWorld* World) override;
+	UE_API virtual void UndoModifier() override;
+	UE_API virtual FText GetModifierName() override;
+	//~ End UMovieGraphModifierBase Interface
+
+public:
+	/** Whether "EnableDOFComponents" state should be applied. */
+	UPROPERTY()
+	bool bOverride_EnableDOFComponents = false;
+
+	/** Whether "NumSamples" state should be applied. */
+	UPROPERTY()
+	bool bOverride_NumSamples = false;
+
+	/** Whether "DOFSplatSize" state should be applied. */
+	UPROPERTY()
+	bool bOverride_DOFSplatSize = false;
+
+	/** The enable/disable state of the Accumulation DOF component. */
+	UPROPERTY()
+	FMovieGraphAccumulationDOFEnableState EnableDOFComponents;
+
+	/** The number of samples that the Accumulation DOF component should use. */
+	UPROPERTY()
+	int32 NumSamples = 0;
+
+	/** The splat size that the Accumulation DOF component should use. */
+	UPROPERTY()
+	float DOFSplatSize = 0.f;
+
+private:
+	/** Tracks the original component state before the modifier was applied. */
+	UPROPERTY(Transient)
+	TArray<FMovieGraphAccumulationDOFModifier_ComponentState> OriginalComponentState;
+};
+
+#undef UE_API
\ No newline at end of file
diff --git a/Engine/Plugins/VirtualProduction/AccumulationDOF/Source/AccumulationDOF/Private/MovieGraphAccumulationDOFModifierCustomization.cpp b/Engine/Plugins/VirtualProduction/AccumulationDOF/Source/AccumulationDOF/Private/MovieGraphAccumulationDOFModifierCustomization.cpp
new file mode 100644
--- /dev/null
+++ b/Engine/Plugins/VirtualProduction/AccumulationDOF/Source/AccumulationDOF/Private/MovieGraphAccumulationDOFModifierCustomization.cpp
@@ -0,0 +1,65 @@
+﻿// Copyright Epic Games, Inc. All Rights Reserved.
+
+#include "MovieGraphAccumulationDOFModifierCustomization.h"
+
+#if WITH_EDITOR
+
+#include "DetailCategoryBuilder.h"
+#include "DetailLayoutBuilder.h"
+#include "DetailWidgetRow.h"
+#include "MovieGraphAccumulationDOFModifierNode.h"
+#include "Widgets/Text/STextBlock.h"
+
+#define LOCTEXT_NAMESPACE "MovieGraphAccumulationDOFModifierCustomization"
+
+TSharedRef<IDetailCustomization> FMovieGraphAccumulationDOFModifierCustomization::MakeInstance()
+{
+	return MakeShared<FMovieGraphAccumulationDOFModifierCustomization>();
+}
+
+void FMovieGraphAccumulationDOFModifierCustomization::CustomizeDetails(IDetailLayoutBuilder& InDetailBuilder)
+{
+	TArray<TWeakObjectPtr<UObject>> CustomizedObjects;
+	InDetailBuilder.GetObjectsBeingCustomized(CustomizedObjects);
+
+	// The customization only supports editing a single Modifier node
+	bool bFoundModifierNode = false;
+	for (const TWeakObjectPtr<UObject>& CustomizedObject : CustomizedObjects)
+	{
+		if (CustomizedObject.IsValid() && CustomizedObject->IsA<UMovieGraphAccumulationDOFModifierNode>())
+		{
+			bFoundModifierNode = true;
+			break;
+		}
+	}
+
+	if (!bFoundModifierNode)
+	{
+		return;
+	}
+
+	// Add a "Collections" category with a custom whole-row widget which indicates that all cameras will be matched
+	IDetailCategoryBuilder& CollectionsCategory =
+		InDetailBuilder.EditCategory(FName("Collections"), FText::GetEmpty(), ECategoryPriority::Important);
+
+	CollectionsCategory.AddCustomRow(FText::GetEmpty())
+	.WholeRowWidget
+	.HAlign(HAlign_Center)
+	.VAlign(VAlign_Center)
+	[
+		SNew(STextBlock)
+		.Margin(FMargin(0, 5, 0, 5))
+		.Text(LOCTEXT("AllCameras", "Modifies the Accumulation DOF component on all active cameras"))
+		.Font(IDetailLayoutBuilder::GetDetailFontItalic())
+	];
+}
+
+void FMovieGraphAccumulationDOFModifierCustomization::CustomizeDetails(const TSharedPtr<IDetailLayoutBuilder>& InDetailBuilder)
+{
+	DetailBuilder = InDetailBuilder;
+	CustomizeDetails(*InDetailBuilder);
+}
+
+#undef LOCTEXT_NAMESPACE
+
+#endif // WITH_EDITOR
diff --git a/Engine/Plugins/VirtualProduction/AccumulationDOF/Source/AccumulationDOF/Private/MovieGraphAccumulationDOFModifierCustomization.h b/Engine/Plugins/VirtualProduction/AccumulationDOF/Source/AccumulationDOF/Private/MovieGraphAccumulationDOFModifierCustomization.h
new file mode 100644
--- /dev/null
+++ b/Engine/Plugins/VirtualProduction/AccumulationDOF/Source/AccumulationDOF/Private/MovieGraphAccumulationDOFModifierCustomization.h
@@ -0,0 +1,27 @@
+﻿// Copyright Epic Games, Inc. All Rights Reserved.
+
+#pragma once
+
+#if WITH_EDITOR
+
+#include "IDetailCustomization.h"
+
+class IMovieGraphModifierNodeInterface;
+
+/** Customize how the Accumulation DOF modifier node appears in the details panel. */
+class FMovieGraphAccumulationDOFModifierCustomization final : public IDetailCustomization
+{
+public:
+	static TSharedRef<IDetailCustomization> MakeInstance();
+
+	//~ Begin IDetailCustomization interface
+	virtual void CustomizeDetails(IDetailLayoutBuilder& InDetailBuilder) override;
+	virtual void CustomizeDetails(const TSharedPtr<IDetailLayoutBuilder>& InDetailBuilder) override;
+	//~ End IDetailCustomization interface
+
+private:
+	/** The details builder associated with the customization. */
+	TWeakPtr<IDetailLayoutBuilder> DetailBuilder;
+};
+
+#endif // WITH_EDITOR
\ No newline at end of file
diff --git a/Engine/Plugins/VirtualProduction/AccumulationDOF/Source/AccumulationDOF/Private/MovieGraphAccumulationDOFModifierNode.cpp b/Engine/Plugins/VirtualProduction/AccumulationDOF/Source/AccumulationDOF/Private/MovieGraphAccumulationDOFModifierNode.cpp
new file mode 100644
--- /dev/null
+++ b/Engine/Plugins/VirtualProduction/AccumulationDOF/Source/AccumulationDOF/Private/MovieGraphAccumulationDOFModifierNode.cpp
@@ -0,0 +1,67 @@
+﻿// Copyright Epic Games, Inc. All Rights Reserved.
+
+#include "MovieGraphAccumulationDOFModifierNode.h"
+
+#include "MovieGraphAccumulationDOFModifier.h"
+
+#include UE_INLINE_GENERATED_CPP_BY_NAME(MovieGraphAccumulationDOFModifierNode)
+
+#define LOCTEXT_NAMESPACE "MovieGraph"
+
+UMovieGraphAccumulationDOFModifierNode::UMovieGraphAccumulationDOFModifierNode()
+{
+	Modifier = CreateDefaultSubobject<UMovieGraphAccumulationDOFModifier>(TEXT("AccumulationDOFModifier"));
+}
+
+#if WITH_EDITOR
+FText UMovieGraphAccumulationDOFModifierNode::GetNodeTitle(const bool bGetDescriptive) const
+{
+	static const FText NodeName = LOCTEXT("NodeName_AccumulationDOFModifier", "Accumulation DOF Modifier");
+
+	return NodeName;
+}
+
+FText UMovieGraphAccumulationDOFModifierNode::GetMenuCategory() const
+{
+	return LOCTEXT("AccumulationDOFModifierNode_Category", "Utility");
+}
+
+FLinearColor UMovieGraphAccumulationDOFModifierNode::GetNodeTitleColor() const
+{
+	static const FLinearColor ModifierNodeColor = FLinearColor(0.6f, 0.113f, 0.113f);
+	return ModifierNodeColor;
+}
+
+FSlateIcon UMovieGraphAccumulationDOFModifierNode::GetIconAndTint(FLinearColor& OutColor) const
+{
+	static const FSlateIcon ModifierIcon = FSlateIcon(FAppStyle::GetAppStyleSetName(), "LevelViewport.ToggleActorPilotCameraView");
+
+	OutColor = FLinearColor::White;
+	return ModifierIcon;
+}
+#endif // WITH_EDITOR
+
+TArray<UMovieGraphModifierBase*> UMovieGraphAccumulationDOFModifierNode::GetAllModifiers() const
+{
+	Modifier->bOverride_EnableDOFComponents = bOverride_EnableAccumulationDepthOfField;
+	Modifier->bOverride_NumSamples = bOverride_NumSamples;
+	Modifier->bOverride_DOFSplatSize = bOverride_DOFSplatSize;
+
+	Modifier->EnableDOFComponents = EnableAccumulationDepthOfField;
+	Modifier->NumSamples = NumSamples;
+	Modifier->DOFSplatSize = DOFSplatSize;
+
+	return { Modifier };
+}
+
+bool UMovieGraphAccumulationDOFModifierNode::SupportsCollections() const
+{
+	return false;
+}
+
+FString UMovieGraphAccumulationDOFModifierNode::GetNodeInstanceName() const
+{
+	return TEXT("AccumulationDOF");
+}
+
+#undef LOCTEXT_NAMESPACE // "MovieGraph"
\ No newline at end of file
diff --git a/Engine/Plugins/VirtualProduction/AccumulationDOF/Source/AccumulationDOF/Private/MovieGraphAccumulationDOFPass.cpp b/Engine/Plugins/VirtualProduction/AccumulationDOF/Source/AccumulationDOF/Private/MovieGraphAccumulationDOFPass.cpp
new file mode 100644
--- /dev/null
+++ b/Engine/Plugins/VirtualProduction/AccumulationDOF/Source/AccumulationDOF/Private/MovieGraphAccumulationDOFPass.cpp
@@ -0,0 +1,893 @@
+﻿// Copyright Epic Games, Inc. All Rights Reserved.
+
+#include "MovieGraphAccumulationDOFPass.h"
+
+#include "AccumulationDOFComponent.h"
+#include "AccumulationDOFMath.h"
+#include "AccumulationDOFSceneViewExtension.h"
+#include "AccumulationDOFShaders.h"
+#include "AccumulationDOFUtils.h"
+#include "ApertureSampler.h"
+
+#include "CanvasTypes.h"
+#include "CineCameraActor.h"
+#include "CineCameraComponent.h"
+#include "Engine/TextureRenderTarget2D.h"
+#include "Engine/World.h"
+#include "EngineModule.h"
+#include "Graph/MovieGraphDefaultRenderer.h"
+#include "Graph/MovieGraphPipeline.h"
+#include "Graph/Nodes/MovieGraphImagePassBaseNode.h"
+#include "MovieRenderPipelineDataTypes.h"
+#include "SceneManagement.h"
+#include "SceneView.h"
+#include "TextureResource.h"
+
+DEFINE_LOG_CATEGORY(LogAccumulationDOFMRG);
+
+FMovieGraphAccumulationDOFPass::FMovieGraphAccumulationDOFPass()
+	: FMovieGraphDeferredPass()
+{
+}
+
+FMovieGraphAccumulationDOFPass::~FMovieGraphAccumulationDOFPass()
+{
+	if (ApertureSampler)
+	{
+		ApertureSampler->Shutdown();
+		ApertureSampler = nullptr;
+	}
+}
+
+void FMovieGraphAccumulationDOFPass::Setup(TWeakObjectPtr<UMovieGraphDefaultRenderer> InRenderer, TWeakObjectPtr<UMovieGraphImagePassBaseNode> InRenderPassNode, const FMovieGraphRenderPassLayerData& InLayer)
+{
+	FMovieGraphDeferredPass::Setup(InRenderer, InRenderPassNode, InLayer);
+
+	const UWorld* World = InRenderer->GetWorld();
+	if (!World)
+	{
+		UE_LOG(LogAccumulationDOFMRG, Warning, TEXT("Invalid world found. Accumulation DOF pass will not be set up properly."));
+		return;
+	}
+
+	// Allocate dedicated view state for injection renders.
+	InjectionViewState.Allocate(World->GetFeatureLevel());
+
+	// Add Lumen scene data initialization
+	if (InjectionViewState.GetReference() && World && World->Scene)
+	{
+		InjectionViewState.GetReference()->AddLumenSceneData(World->Scene, 1.0f);
+	}
+
+	// Override LUT size to minimize banding
+	constexpr bool bForceColorGradingLUTSize64 = true;
+	SavedLUTSize = AccumulationDOFUtils::OverrideLUTSizeIfNeeded(bForceColorGradingLUTSize64);
+}
+
+void FMovieGraphAccumulationDOFPass::Teardown()
+{
+	if (ApertureSampler)
+	{
+		ApertureSampler->Shutdown();
+		ApertureSampler = nullptr;
+	}
+
+	// Restore LUT size if it was overridden
+	AccumulationDOFUtils::RestoreLUTSize(SavedLUTSize);
+	SavedLUTSize = 0;
+
+	CachedCineCameraComponent.Reset();
+	CachedDOFComponent.Reset();
+	CaptureExtension.Reset();
+	InjectionExtension.Reset();
+	InjectionViewState.Destroy();
+
+	FMovieGraphDeferredPass::Teardown();
+}
+
+void FMovieGraphAccumulationDOFPass::AddReferencedObjects(FReferenceCollector& Collector)
+{
+	FMovieGraphDeferredPass::AddReferencedObjects(Collector);
+
+	// Protect ApertureSampler (and its UPROPERTY render targets) from GC.
+	// This is critical because ForceGarbageCollection() is called during the
+	// aperture sampling loop via ReleaseRenderingMemory().
+	Collector.AddReferencedObject(ApertureSampler);
+
+	// Protect the InjectionViewState from garbage collection
+	FSceneViewStateInterface* InjectionViewRef = InjectionViewState.GetReference();
+	if (InjectionViewRef)
+	{
+		InjectionViewRef->AddReferencedObjects(Collector);
+	}
+}
+
+void FMovieGraphAccumulationDOFPass::UpdateCameraDependentState()
+{
+	// Update cached camera component
+	if (!CachedCineCameraComponent.IsValid())
+	{
+		CachedCineCameraComponent = FindCineCameraComponent();
+		FindAccumulationDOFComponent();
+	}
+
+	// Update sampler's camera state
+	if (ApertureSampler && CachedCineCameraComponent.IsValid())
+	{
+		AccumulationDOF::FApertureSamplerCameraState CameraState = BuildCameraState();
+		ApertureSampler->UpdateCameraState(CameraState);
+		bCameraParamsInitialized = true;
+	}
+	else if (!bCameraParamsInitialized)
+	{
+		UE_LOG(LogAccumulationDOFMRG, Warning, TEXT("No CineCameraActor found. Using default values."));
+		bCameraParamsInitialized = true;
+	}
+}
+
+void FMovieGraphAccumulationDOFPass::Render(const FMovieGraphTraversalContext& InFrameTraversalContext, const FMovieGraphTimeStepData& InTimeData)
+{
+	// Initialize the sampler first. This would ideally be done in Setup(), but we need access to the evaluated graph for this.
+	if (!ApertureSampler)
+	{
+		InitializeSampler(InTimeData.EvaluatedConfig.Get());
+	}
+
+	CachedMotionBlurFraction = InTimeData.MotionBlurFraction;
+
+	UpdateCameraDependentState();
+
+	if (!ApertureSampler || !ApertureSampler->IsInitialized())
+	{
+		UE_LOG(LogAccumulationDOFMRG, Error, TEXT("ApertureSampler not initialized"));
+		return;
+	}
+
+	TObjectPtr<UMovieGraphDefaultRenderer> GraphRenderer = GetRenderer().Get();
+	if (!GraphRenderer)
+	{
+		return;
+	}
+
+	const FResolutionAndCameraInfo ResolutionAndCameraInfo = GetResolutionAndCameraInfo(InTimeData.EvaluatedConfig.Get());
+	const UE::MovieGraph::DefaultRenderer::FRenderTargetInitParams RenderTargetInitParams = GetRenderTargetInitParams(InTimeData, ResolutionAndCameraInfo.BackbufferResolution);
+
+	// During warmup frames, render a single aperture sample to prime temporal history
+	if (InTimeData.bDiscardOutput)
+	{
+		CurrentApertureSampleIndex = 0;
+
+		UMovieGraphImagePassBaseNode* ParentNodeThisFrame = GetParentNode(InTimeData.EvaluatedConfig);
+
+		UTextureRenderTarget2D* RenderTarget = GraphRenderer->GetOrCreateViewRenderTarget(RenderTargetInitParams, RenderDataIdentifier);
+		FRenderTarget* RenderTargetResource = RenderTarget->GameThread_GetRenderTargetResource();
+		check(RenderTargetResource);
+
+		UE::MovieGraph::DefaultRenderer::FCameraInfo CameraInfo = GraphRenderer->GetCameraInfo(LayerData.CameraIndex);
+		CameraInfo.bAllowCameraAspectRatio = true;
+		CameraInfo.TilingParams.TileCount = FIntPoint(1, 1);
+		CameraInfo.TilingParams.TileIndexes = FIntPoint(0, 0);
+		CameraInfo.TilingParams.TileSize = ResolutionAndCameraInfo.TileSize;
+		CameraInfo.ProjectionMatrixJitterAmount = ComputeProjectionMatrixJitter(FVector2D::ZeroVector, ResolutionAndCameraInfo.BackbufferResolution);
+		CameraInfo.bUseCameraManagerPostProcess = !ResolutionAndCameraInfo.bRenderAllCameras;
+		CameraInfo.ViewInfo.ClearOverscan();
+		CameraInfo.ViewInfo.ApplyOverscan(ResolutionAndCameraInfo.OverscanFraction);
+
+		UE::MovieGraph::Rendering::FViewFamilyInitData ViewFamilyInitData;
+		ViewFamilyInitData.RenderTarget = RenderTargetResource;
+		ViewFamilyInitData.World = GraphRenderer->GetWorld();
+		ViewFamilyInitData.TimeData = InTimeData;
+		ViewFamilyInitData.SceneCaptureSource = ESceneCaptureSource::SCS_FinalToneCurveHDR;
+		ViewFamilyInitData.bWorldIsPaused = false;
+		ViewFamilyInitData.FrameIndex = InTimeData.RenderedFrameNumber;
+		ViewFamilyInitData.AntiAliasingMethod = EAntiAliasingMethod::AAM_None;
+		ViewFamilyInitData.ShowFlags = ParentNodeThisFrame ? ParentNodeThisFrame->GetShowFlags() : FEngineShowFlags(ESFIM_Game);
+		ViewFamilyInitData.ViewModeIndex = ParentNodeThisFrame ? ParentNodeThisFrame->GetViewModeIndex() : VMI_Lit;
+		ViewFamilyInitData.ProjectionMode = CameraInfo.ViewInfo.ProjectionMode;
+
+		FCreateViewResult ViewResult = CreateConfiguredView(CameraInfo, ViewFamilyInitData, ResolutionAndCameraInfo.BackbufferResolution, ResolutionAndCameraInfo.AccumulatorResolution, ParentNodeThisFrame, 0, 0);
+
+		FCanvas Canvas = FCanvas(RenderTargetResource, nullptr, GraphRenderer->GetWorld(), GraphRenderer->GetWorld()->GetFeatureLevel(), FCanvas::CDM_DeferDrawing, 1.0f);
+		GetRendererModule().BeginRenderingViewFamily(&Canvas, ViewResult.ViewFamily.ToSharedPtr().Get());
+
+		CurrentApertureSampleIndex = -1;
+		return;
+	}
+
+	// Full render: all aperture samples + injection + submission
+	RenderSamples(InFrameTraversalContext, InTimeData, ResolutionAndCameraInfo, RenderTargetInitParams);
+}
+
+TSharedRef<FSceneViewFamilyContext> FMovieGraphAccumulationDOFPass::CreateSceneViewFamily(const UE::MovieGraph::Rendering::FViewFamilyInitData& InInitData) const
+{
+	TSharedRef<FSceneViewFamilyContext> ViewFamily = FMovieGraphDeferredPass::CreateSceneViewFamily(InInitData);
+
+	// DOF splats require engine DOF to be enabled
+	const bool bDOFSplatsEnabled = GetEffectiveDOFSplatSize() > UE_KINDA_SMALL_NUMBER;
+	ViewFamily->EngineShowFlags.SetDepthOfField(bDOFSplatsEnabled);
+
+	// Motion blur is only enabled during injection pass when requested
+	const bool bWantsMotionBlur = CachedDOFComponent.IsValid() && CachedDOFComponent->bEnableMotionBlur;
+	const bool bMotionBlurEnabled = bWantsMotionBlur && !FMath::IsNearlyZero(CachedMotionBlurFraction);
+	ViewFamily->EngineShowFlags.SetMotionBlur(bMotionBlurEnabled);
+
+	// AA show flags: when bUseJitterAA is true (default), TAA show flags are disabled so the pass
+	// uses its own jitter-based AA via aperture sampling. When false, TAA show flags are enabled
+	// and actual TAA processing is controlled via AntiAliasingMethod in ApplyMovieGraphOverridesToSceneView.
+	const bool bUseJitterAA = CachedDOFComponent.IsValid() ? CachedDOFComponent->bUseJitterAA : true;
+	ViewFamily->EngineShowFlags.SetTemporalAA(!bUseJitterAA);
+	ViewFamily->EngineShowFlags.SetAntiAliasing(!bUseJitterAA);
+	
+	// Add AccumulationDOF scene view extensions (equivalent to MRQ's AddViewExtensions)
+	if (bInCapturePass && CaptureExtension.IsValid())
+	{
+		ViewFamily->ViewExtensions.Add(CaptureExtension.ToSharedRef());
+	}
+
+	if (bInInjectionPass && InjectionExtension.IsValid())
+	{
+		ViewFamily->ViewExtensions.Add(InjectionExtension.ToSharedRef());
+	}
+
+	return ViewFamily;
+}
+
+FVector2D FMovieGraphAccumulationDOFPass::ComputeProjectionMatrixJitter(const FVector2D& InDefaultJitter, const FIntPoint& InBackbufferResolution) const
+{
+	if (ApertureSampler && CurrentApertureSampleIndex >= 0)
+	{
+		return ApertureSampler->GetJitterForProjectionMatrix(CurrentApertureSampleIndex, InBackbufferResolution);
+	}
+
+	return FVector2D::ZeroVector;
+}
+
+FSceneViewInitOptions FMovieGraphAccumulationDOFPass::CreateViewInitOptions(
+	const UE::MovieGraph::DefaultRenderer::FCameraInfo& InCameraInfo,
+	FSceneViewFamilyContext* InViewFamily,
+	FSceneViewStateInterface* InViewStateInterface) const
+{
+	FSceneViewInitOptions ViewInitOptions = FMovieGraphDeferredPass::CreateViewInitOptions(InCameraInfo, InViewFamily, InViewStateInterface);
+
+	// Apply aperture sample camera offset in the lens plane
+	if (ApertureSampler && CurrentApertureSampleIndex >= 0 && CurrentApertureSampleIndex < ApertureSampler->GetActualNumSamples())
+	{
+		const FVector2f ApertureOffset = ApertureSampler->GetApertureOffset(CurrentApertureSampleIndex);
+
+		const FRotator CameraRotation = InCameraInfo.ViewInfo.Rotation;
+		const FVector CameraRight = FRotationMatrix(CameraRotation).GetUnitAxis(EAxis::Y);
+		const FVector CameraUp = FRotationMatrix(CameraRotation).GetUnitAxis(EAxis::Z);
+
+		const FVector OffsetLocation = InCameraInfo.ViewInfo.Location + CameraRight * ApertureOffset.X + CameraUp * ApertureOffset.Y;
+
+		ViewInitOptions.ViewOrigin = OffsetLocation;
+		ViewInitOptions.ViewLocation = OffsetLocation;
+	}
+
+	return ViewInitOptions;
+}
+
+void FMovieGraphAccumulationDOFPass::CalculateProjectionMatrix(
+	UE::MovieGraph::DefaultRenderer::FCameraInfo& InOutCameraInfo,
+	FSceneViewProjectionData& InOutProjectionData,
+	const FIntPoint InBackbufferResolution,
+	const FIntPoint InAccumulatorResolution) const
+{
+	if (ApertureSampler && CurrentApertureSampleIndex >= 0 && CurrentApertureSampleIndex < ApertureSampler->GetActualNumSamples())
+	{
+		FMatrix ProjectionMatrix;
+
+		if (ApertureSampler->ComputeProjectionMatrix(CurrentApertureSampleIndex, CurrentFocusDistanceOverride, ProjectionMatrix))
+		{
+			InOutProjectionData.ProjectionMatrix = ProjectionMatrix;
+
+			// Reset PreviousViewTransform — with aperture sampling, keeping the center position
+			// gives Lumen false camera movement causing what looks like reprojection drift.
+			InOutCameraInfo.ViewInfo.PreviousViewTransform.Reset();
+			return;
+		}
+	}
+
+	// Fallback to default projection matrix calculation
+	FMovieGraphDeferredPass::CalculateProjectionMatrix(InOutCameraInfo, InOutProjectionData, InBackbufferResolution, InAccumulatorResolution);
+}
+
+void FMovieGraphAccumulationDOFPass::ApplyMovieGraphOverridesToSceneView(
+	TSharedRef<FSceneViewFamilyContext> InOutFamily,
+	const UE::MovieGraph::Rendering::FViewFamilyInitData& InInitData,
+	const UE::MovieGraph::DefaultRenderer::FCameraInfo& InCameraInfo) const
+{
+	FMovieGraphDeferredPass::ApplyMovieGraphOverridesToSceneView(InOutFamily, InInitData, InCameraInfo);
+
+	FSceneView* View = const_cast<FSceneView*>(InOutFamily->Views[0]);
+
+	// Override AA method: use TAA only when bUseJitterAA is false and view state is available
+	const bool bUseJitterAA = CachedDOFComponent.IsValid() ? CachedDOFComponent->bUseJitterAA : true;
+	const bool bCanUseTAA = !bUseJitterAA && View->State != nullptr;
+	View->AntiAliasingMethod = bCanUseTAA ? AAM_TemporalAA : AAM_None;
+
+	// Disable motion blur for capture passes — only the injection pass should have motion blur
+	if (!bInInjectionPass)
+	{
+		View->FinalPostProcessSettings.MotionBlurAmount = 0.0f;
+		View->FinalPostProcessSettings.bOverride_MotionBlurAmount = true;
+	}
+
+	// During injection, disable engine's lateral chromatic aberration if we applied our own spectral version
+	if (bInInjectionPass)
+	{
+		const bool bUseSpectralLateralCA = CachedDOFComponent.IsValid()
+			? CachedDOFComponent->bSpectralLateralChromaticAberration
+			: true;
+
+		const bool bAppliedOurSpectralCA = bUseSpectralLateralCA && (CachedSceneFringeIntensity > UE_KINDA_SMALL_NUMBER);
+
+		if (bAppliedOurSpectralCA)
+		{
+			View->FinalPostProcessSettings.SceneFringeIntensity = 0.0f;
+		}
+	}
+}
+
+FSceneViewStateInterface* FMovieGraphAccumulationDOFPass::GetSceneViewState(UMovieGraphImagePassBaseNode* ParentNodeThisFrame, int32_t TileX, int32_t TileY)
+{
+	// During injection pass, always use dedicated injection view state
+
+	if (bInInjectionPass && InjectionViewState.GetReference())
+	{
+		return InjectionViewState.GetReference();
+	}
+
+	FSceneViewStateInterface* BaseState = FMovieGraphDeferredPass::GetSceneViewState(ParentNodeThisFrame, TileX, TileY);
+
+	return BaseState;
+}
+
+void FMovieGraphAccumulationDOFPass::RenderSamples(
+	const FMovieGraphTraversalContext& InFrameTraversalContext,
+	const FMovieGraphTimeStepData& InTimeData,
+	const FResolutionAndCameraInfo& InResolutionInfo,
+	const UE::MovieGraph::DefaultRenderer::FRenderTargetInitParams& InRenderTargetInitParams)
+{
+	const int32 ActualNumSamples = ApertureSampler->GetActualNumSamples();
+
+	// Reset accumulation for new frame
+	ApertureSampler->ResetAccumulation();
+
+	// Check for axial chromatic aberration
+	const float AxialCAIntensity = CachedDOFComponent.IsValid()
+		? CachedDOFComponent->AxialChromaticAberrationIntensity
+		: 0.0f;
+
+	const int32 EffectiveBands = ApertureSampler->GetEffectiveBands();
+	const bool bHasAxialCA = (AxialCAIntensity > UE_KINDA_SMALL_NUMBER) && (EffectiveBands > 1);
+
+	if (bHasAxialCA)
+	{
+		// Axial CA: render each sample with varying focus distances per spectral band
+		const FVector3f TotalSpectralWeight = ApertureSampler->GetTotalSpectralWeightPerChannel();
+
+		const FVector3f InvTotal(
+			(TotalSpectralWeight.X > UE_KINDA_SMALL_NUMBER) ? (1.0f / TotalSpectralWeight.X) : 0.0f,
+			(TotalSpectralWeight.Y > UE_KINDA_SMALL_NUMBER) ? (1.0f / TotalSpectralWeight.Y) : 0.0f,
+			(TotalSpectralWeight.Z > UE_KINDA_SMALL_NUMBER) ? (1.0f / TotalSpectralWeight.Z) : 0.0f
+		);
+
+		// Get nominal focus distance for spectral calculations
+		const float NominalFocusDistanceCm = CachedCineCameraComponent.IsValid()
+			? CachedCineCameraComponent->CurrentFocusDistance
+			: 1000.0f;
+
+		int32 IterationCount = 0;
+
+		for (int32 SampleIndex = 0; SampleIndex < ActualNumSamples; ++SampleIndex)
+		{
+			for (int32 Band = 0; Band < EffectiveBands; ++Band)
+			{
+				CurrentApertureSampleIndex = SampleIndex;
+
+				// Compute spectral parameters for this band
+				const float NormalizedSpectralPosition = (EffectiveBands == 1)
+					? 0.5f
+					: static_cast<float>(Band) / static_cast<float>(EffectiveBands - 1);
+
+				CurrentFocusDistanceOverride = ApertureSampler->ComputeSpectralFocusDistance(
+					NominalFocusDistanceCm, NormalizedSpectralPosition);
+
+				const FVector3f RawWeight = ApertureSampler->ComputeSpectralWeight(NormalizedSpectralPosition);
+				CurrentAxialCAWeight = RawWeight * InvTotal;
+
+				RenderSingleApertureSample(SampleIndex, InFrameTraversalContext, InTimeData, InResolutionInfo, InRenderTargetInitParams);
+
+				IterationCount++;
+
+				// Release memory periodically
+				if ((IterationCount % AccumulationDOFUtils::RenderingMemoryFlushBatchSize) == 0)
+				{
+					AccumulationDOFUtils::ReleaseRenderingMemory();
+				}
+			}
+		}
+	}
+	else
+	{
+		// No axial CA: standard rendering with uniform spectral weight
+		CurrentFocusDistanceOverride = -1.0f;
+		CurrentAxialCAWeight = FVector3f(1.0f);
+
+		for (int32 SampleIndex = 0; SampleIndex < ActualNumSamples; ++SampleIndex)
+		{
+			CurrentApertureSampleIndex = SampleIndex;
+			RenderSingleApertureSample(SampleIndex, InFrameTraversalContext, InTimeData, InResolutionInfo, InRenderTargetInitParams);
+		}
+	}
+
+	// Reset state after all samples rendered
+	CurrentApertureSampleIndex = -1;
+	CurrentFocusDistanceOverride = -1.0f;
+	CurrentAxialCAWeight = FVector3f(1.0f);
+	CaptureExtension.Reset();
+
+	// Update camera state with captured SceneFringeIntensity for lateral CA
+	ApertureSampler->UpdateCameraState(BuildCameraState());
+
+	// Finalize accumulation
+	ApertureSampler->FinalizeAccumulation();
+
+	// Render injection pass and submit to MRG
+	RenderInjectionAndSubmit(InFrameTraversalContext, InTimeData, InResolutionInfo, InRenderTargetInitParams);
+}
+
+void FMovieGraphAccumulationDOFPass::RenderSingleApertureSample(
+	int32 SampleIndex,
+	const FMovieGraphTraversalContext& InFrameTraversalContext,
+	const FMovieGraphTimeStepData& InTimeData,
+	const FResolutionAndCameraInfo& InResolutionInfo,
+	const UE::MovieGraph::DefaultRenderer::FRenderTargetInitParams& InRenderTargetInitParams)
+{
+	TObjectPtr<UMovieGraphDefaultRenderer> GraphRenderer = GetRenderer().Get();
+	if (!GraphRenderer)
+	{
+		return;
+	}
+
+	UMovieGraphImagePassBaseNode* ParentNodeThisFrame = GetParentNode(InTimeData.EvaluatedConfig);
+	const FIntPoint BackbufferResolution = InResolutionInfo.BackbufferResolution;
+	const FIntPoint AccumulatorResolution = InResolutionInfo.AccumulatorResolution;
+
+	// Create SVE to capture scene color
+	UTextureRenderTarget2D* SampleRT = ApertureSampler->GetSampleRT();
+
+	if (SampleRT)
+	{
+		FAccumulationDOFSVESettings SVESettings = FAccumulationDOFSVESettings::FromCineCameraComponent(CachedCineCameraComponent.Get());
+
+		CaptureExtension = MakeShared<FAccumulationDOFSceneViewExtension, ESPMode::ThreadSafe>(
+			SampleRT,
+			EAccumulationDOFSVEMode::Capture,
+			SVESettings
+		);
+
+		// Apply DOF splats if size is > 0
+		const float EffectiveDOFSplatSize = GetEffectiveDOFSplatSize();
+
+		if (EffectiveDOFSplatSize > UE_KINDA_SMALL_NUMBER)
+		{
+			const AccumulationDOFUtils::FCameraParams& CamParams = ApertureSampler->GetConfig().NumSamples > 0 ?
+				AccumulationDOFUtils::ExtractCameraParams(CachedCineCameraComponent.Get()) :
+				AccumulationDOFUtils::FCameraParams();
+
+			const float DOFSplatsFStop = AccumulationDOFMath::ComputeDOFSplatsFStop(
+				CamParams.ApertureRadiusCm,
+				CamParams.FocalLengthMm,
+				EffectiveDOFSplatSize
+			);
+
+			const float SafeSqueezeFactor = FMath::Max(CamParams.SqueezeFactor, UE_KINDA_SMALL_NUMBER);
+			if (DOFSplatsFStop > UE_KINDA_SMALL_NUMBER)
+			{
+				const float SensorWidthMm = CamParams.SensorWidthMm / SafeSqueezeFactor;
+				const bool bForceNeutralBokeh = true;
+
+				CaptureExtension->SetDOFSplatsSettings(DOFSplatsFStop, CamParams.FocusDistanceCm, SensorWidthMm, CamParams.SqueezeFactor, bForceNeutralBokeh);
+			}
+		}
+	}
+	else
+	{
+		CaptureExtension.Reset();
+	}
+
+	// Get render target
+	UTextureRenderTarget2D* RenderTarget = GraphRenderer->GetOrCreateViewRenderTarget(InRenderTargetInitParams, RenderDataIdentifier);
+	FRenderTarget* RenderTargetResource = RenderTarget->GameThread_GetRenderTargetResource();
+	check(RenderTargetResource);
+
+	// Build camera info
+	UE::MovieGraph::DefaultRenderer::FCameraInfo CameraInfo = GraphRenderer->GetCameraInfo(LayerData.CameraIndex);
+	CameraInfo.bAllowCameraAspectRatio = true;
+	CameraInfo.TilingParams.TileSize = InResolutionInfo.TileSize;
+	CameraInfo.TilingParams.TileCount = FIntPoint(1, 1);
+	CameraInfo.TilingParams.TileIndexes = FIntPoint(0, 0);
+	CameraInfo.SamplingParams.TemporalSampleIndex = InTimeData.TemporalSampleIndex;
+	CameraInfo.SamplingParams.TemporalSampleCount = InTimeData.TemporalSampleCount;
+	CameraInfo.SamplingParams.SpatialSampleIndex = 0;
+	CameraInfo.SamplingParams.SpatialSampleCount = 1;
+	CameraInfo.ProjectionMatrixJitterAmount = ComputeProjectionMatrixJitter(FVector2D::ZeroVector, BackbufferResolution);
+	CameraInfo.bUseCameraManagerPostProcess = !InResolutionInfo.bRenderAllCameras;
+	CameraInfo.ViewInfo.ClearOverscan();
+	CameraInfo.ViewInfo.ApplyOverscan(InResolutionInfo.OverscanFraction);
+
+	// Build view family init data
+	const int32 FrameIndex = InTimeData.RenderedFrameNumber;
+	const ESceneCaptureSource SceneCaptureSource = ParentNodeThisFrame->GetDisableToneCurve() ? ESceneCaptureSource::SCS_FinalColorHDR : ESceneCaptureSource::SCS_FinalToneCurveHDR;
+
+	UE::MovieGraph::Rendering::FViewFamilyInitData ViewFamilyInitData;
+	ViewFamilyInitData.RenderTarget = RenderTargetResource;
+	ViewFamilyInitData.World = GraphRenderer->GetWorld();
+	ViewFamilyInitData.TimeData = InTimeData;
+	ViewFamilyInitData.SceneCaptureSource = SceneCaptureSource;
+	ViewFamilyInitData.bWorldIsPaused = false; // Will be set below based on TemporalHistoryMode
+	ViewFamilyInitData.FrameIndex = FrameIndex;
+	ViewFamilyInitData.AntiAliasingMethod = EAntiAliasingMethod::AAM_None;
+	ViewFamilyInitData.ShowFlags = ParentNodeThisFrame->GetShowFlags();
+	ViewFamilyInitData.ViewModeIndex = ParentNodeThisFrame->GetViewModeIndex();
+	ViewFamilyInitData.ProjectionMode = CameraInfo.ViewInfo.ProjectionMode;
+
+	bInCapturePass = true;
+	FCreateViewResult ViewResult = CreateConfiguredView(CameraInfo, ViewFamilyInitData, BackbufferResolution, AccumulatorResolution, ParentNodeThisFrame, 0, 0);
+	bInCapturePass = false;
+
+	TSharedRef<FSceneViewFamilyContext> ViewFamily = ViewResult.ViewFamily;
+	FSceneView* View = ViewResult.View;
+
+	if (!View)
+	{
+		UE_LOG(LogAccumulationDOFMRG, Error, TEXT("Failed to create view for aperture sample %d"), SampleIndex);
+		CaptureExtension.Reset();
+		return;
+	}
+
+	// Modify OverrideFrameIndexValue — combine with aperture sample index for uniqueness
+	const int32 ActualNumSamples = ApertureSampler->GetActualNumSamples();
+	if (View->OverrideFrameIndexValue.IsSet())
+	{
+		const int64 BaseFrameIndex = View->OverrideFrameIndexValue.GetValue();
+		const int64 CombinedIndex = BaseFrameIndex * static_cast<int64>(ActualNumSamples) + SampleIndex;
+		View->OverrideFrameIndexValue = static_cast<int32>(CombinedIndex & 0x7FFFFFFF);
+	}
+
+	// Control when temporal history updates with bWorldIsPaused
+	const bool bIsFirstApertureSample = (SampleIndex == 0);
+	const bool bIsLastApertureSample = (SampleIndex == ActualNumSamples - 1);
+
+	const ETemporalHistoryMode HistoryMode =
+		CachedDOFComponent.IsValid() ? CachedDOFComponent->TemporalHistoryMode : ETemporalHistoryMode::LastSampleOnly;
+
+	switch (HistoryMode)
+	{
+	case ETemporalHistoryMode::AllSamplesUpdate:
+		ViewFamily->bWorldIsPaused = false;
+		break;
+	case ETemporalHistoryMode::FirstSampleOnly:
+		ViewFamily->bWorldIsPaused = !bIsFirstApertureSample;
+		break;
+	case ETemporalHistoryMode::LastSampleOnly:
+		ViewFamily->bWorldIsPaused = !bIsLastApertureSample;
+		break;
+	case ETemporalHistoryMode::NoSamplesUpdate:
+		ViewFamily->bWorldIsPaused = true;
+		break;
+	default:
+		checkNoEntry();
+		break;
+	}
+
+	// Render
+	FCanvas Canvas = FCanvas(RenderTargetResource, nullptr, GraphRenderer->GetWorld(), GraphRenderer->GetWorld()->GetFeatureLevel(), FCanvas::CDM_DeferDrawing, 1.0f);
+	GetRendererModule().BeginRenderingViewFamily(&Canvas, ViewFamily.ToSharedPtr().Get());
+
+	// Capture SceneFringeIntensity from first sample
+	if (SampleIndex == 0 && CaptureExtension.IsValid())
+	{
+		const float CapturedValue = CaptureExtension->GetCapturedSceneFringeIntensity();
+		if (FMath::Abs(CapturedValue - CachedSceneFringeIntensity) > UE_KINDA_SMALL_NUMBER)
+		{
+			CachedSceneFringeIntensity = CapturedValue;
+		}
+	}
+
+	// Accumulate with spectral weight (set per band for axial CA)
+	const FVector2f ApertureOffset = ApertureSampler->GetApertureOffset(SampleIndex);
+	ApertureSampler->AccumulateExternalSample(SampleIndex, ApertureOffset, CurrentAxialCAWeight);
+
+	// Release memory periodically
+	if (((SampleIndex + 1) % AccumulationDOFUtils::RenderingMemoryFlushBatchSize) == 0)
+	{
+		AccumulationDOFUtils::ReleaseRenderingMemory();
+	}
+
+	CaptureExtension.Reset();
+}
+
+void FMovieGraphAccumulationDOFPass::RenderInjectionAndSubmit(
+	const FMovieGraphTraversalContext& InFrameTraversalContext,
+	const FMovieGraphTimeStepData& InTimeData,
+	const FResolutionAndCameraInfo& InResolutionInfo,
+	const UE::MovieGraph::DefaultRenderer::FRenderTargetInitParams& InRenderTargetInitParams)
+{
+	TObjectPtr<UMovieGraphDefaultRenderer> GraphRenderer = GetRenderer().Get();
+	if (!GraphRenderer)
+	{
+		return;
+	}
+
+	UMovieGraphImagePassBaseNode* ParentNodeThisFrame = GetParentNode(InTimeData.EvaluatedConfig);
+	const FIntPoint BackbufferResolution = InResolutionInfo.BackbufferResolution;
+	const FIntPoint AccumulatorResolution = InResolutionInfo.AccumulatorResolution;
+
+	// Get our accumulated result for injection
+	UTextureRenderTarget2D* NormalizedRT = ApertureSampler->GetAccumulatedResult();
+
+	if (!NormalizedRT)
+	{
+		UE_LOG(LogAccumulationDOFMRG, Error, TEXT("No normalized render target available for injection pass"));
+		return;
+	}
+
+	// Determine motion blur setting
+	const bool bWantsMotionBlur = CachedDOFComponent.IsValid() && CachedDOFComponent->bEnableMotionBlur;
+	const bool bMotionBlurEnabled = bWantsMotionBlur && !FMath::IsNearlyZero(CachedMotionBlurFraction);
+
+	// Create injection extension
+	{
+		const EAccumulationDOFSVEMode InjectionMode = bMotionBlurEnabled
+			? EAccumulationDOFSVEMode::InjectViaTemporalUpscaler
+			: EAccumulationDOFSVEMode::Inject;
+
+		InjectionExtension = MakeShared<FAccumulationDOFSceneViewExtension, ESPMode::ThreadSafe>(
+			NormalizedRT,
+			InjectionMode
+		);
+
+		const bool bUseSpectralLateralCA =
+			CachedDOFComponent.IsValid() ?
+			CachedDOFComponent->bSpectralLateralChromaticAberration
+			: true;
+
+		const bool bAppliedOurSpectralCA = bUseSpectralLateralCA && (CachedSceneFringeIntensity > UE_KINDA_SMALL_NUMBER);
+		InjectionExtension->SetAllowSceneFringe(!bAppliedOurSpectralCA);
+	}
+
+	// Get render target
+	UTextureRenderTarget2D* RenderTarget = GraphRenderer->GetOrCreateViewRenderTarget(InRenderTargetInitParams, RenderDataIdentifier);
+	FRenderTarget* RenderTargetResource = RenderTarget->GameThread_GetRenderTargetResource();
+	check(RenderTargetResource);
+
+	// Build camera info for injection
+	UE::MovieGraph::DefaultRenderer::FCameraInfo CameraInfo = GraphRenderer->GetCameraInfo(LayerData.CameraIndex);
+	CameraInfo.bAllowCameraAspectRatio = true;
+	CameraInfo.TilingParams.TileSize = InResolutionInfo.TileSize;
+	CameraInfo.TilingParams.TileCount = FIntPoint(1, 1);
+	CameraInfo.TilingParams.TileIndexes = FIntPoint(0, 0);
+	CameraInfo.SamplingParams.TemporalSampleIndex = InTimeData.TemporalSampleIndex;
+	CameraInfo.SamplingParams.TemporalSampleCount = InTimeData.TemporalSampleCount;
+	CameraInfo.SamplingParams.SpatialSampleIndex = 0;
+	CameraInfo.SamplingParams.SpatialSampleCount = 1;
+	CameraInfo.ProjectionMatrixJitterAmount = FVector2D::ZeroVector;
+	CameraInfo.bUseCameraManagerPostProcess = !InResolutionInfo.bRenderAllCameras;
+	CameraInfo.ViewInfo.ClearOverscan();
+	CameraInfo.ViewInfo.ApplyOverscan(InResolutionInfo.OverscanFraction);
+
+	// Build view family init data
+	const ESceneCaptureSource SceneCaptureSource = ParentNodeThisFrame->GetDisableToneCurve() ? ESceneCaptureSource::SCS_FinalColorHDR : ESceneCaptureSource::SCS_FinalToneCurveHDR;
+
+	UE::MovieGraph::Rendering::FViewFamilyInitData ViewFamilyInitData;
+	ViewFamilyInitData.RenderTarget = RenderTargetResource;
+	ViewFamilyInitData.World = GraphRenderer->GetWorld();
+	ViewFamilyInitData.TimeData = InTimeData;
+	ViewFamilyInitData.SceneCaptureSource = SceneCaptureSource;
+	ViewFamilyInitData.bWorldIsPaused = false;
+	ViewFamilyInitData.FrameIndex = InTimeData.RenderedFrameNumber;
+	ViewFamilyInitData.AntiAliasingMethod = EAntiAliasingMethod::AAM_None;
+	ViewFamilyInitData.ShowFlags = ParentNodeThisFrame->GetShowFlags();
+	ViewFamilyInitData.ViewModeIndex = ParentNodeThisFrame->GetViewModeIndex();
+	ViewFamilyInitData.ProjectionMode = CameraInfo.ViewInfo.ProjectionMode;
+
+	bInInjectionPass = true;
+	FCreateViewResult ViewResult = CreateConfiguredView(CameraInfo, ViewFamilyInitData, BackbufferResolution, AccumulatorResolution, ParentNodeThisFrame, 0, 0);
+	bInInjectionPass = false;
+
+	TSharedRef<FSceneViewFamilyContext> ViewFamily = ViewResult.ViewFamily;
+	FSceneView* View = ViewResult.View;
+
+	if (!View)
+	{
+		InjectionExtension.Reset();
+		UE_LOG(LogAccumulationDOFMRG, Error, TEXT("Failed to create view for injection pass"));
+		return;
+	}
+
+	// On the first injection render, InjectionViewState has no temporal history (warmup only
+	// renders through the main SceneViewStates). Set bCameraCut to signal the renderer to
+	// initialize temporal effects (eye adaptation, etc.) from scratch instead of reading
+	// uninitialized history, which would produce a black frame.
+	if (bFirstInjection)
+	{
+		View->bCameraCut = true;
+		bFirstInjection = false;
+	}
+
+	// Render injection
+	FCanvas Canvas = FCanvas(RenderTargetResource, nullptr, GraphRenderer->GetWorld(), GraphRenderer->GetWorld()->GetFeatureLevel(), FCanvas::CDM_DeferDrawing, 1.0f);
+	GetRendererModule().BeginRenderingViewFamily(&Canvas, ViewFamily.ToSharedPtr().Get());
+
+	InjectionExtension.Reset();
+
+	// Build sample state and submit — DOF calls PostRendererSubmission directly (no FIFO queue delay)
+	UE::MovieGraph::FMovieGraphSampleState SampleState;
+	{
+		FMovieGraphTraversalContext UpdatedTraversalContext = InFrameTraversalContext;
+		UpdatedTraversalContext.Time = InTimeData;
+		UpdatedTraversalContext.Time.SpatialSampleIndex = 0;
+		UpdatedTraversalContext.Time.SpatialSampleCount = 1;
+		UpdatedTraversalContext.RenderDataIdentifier = RenderDataIdentifier;
+
+		SampleState.TraversalContext = MoveTemp(UpdatedTraversalContext);
+		SampleState.OverscannedResolution = AccumulatorResolution;
+		SampleState.UnpaddedTileSize = InResolutionInfo.TileSize;
+		SampleState.BackbufferResolution = BackbufferResolution;
+		SampleState.AccumulatorResolution = AccumulatorResolution;
+		SampleState.bWriteSampleToDisk = false;
+		SampleState.bRequiresAccumulator = InTimeData.bRequiresAccumulator;
+		SampleState.bFetchFromAccumulator = InTimeData.bIsLastTemporalSampleForFrame;
+		SampleState.OverlappedPad = FIntPoint(0, 0);
+		SampleState.OverlappedOffset = FIntPoint(0, 0);
+		SampleState.OverlappedSubpixelShift = FVector2D(0.5, 0.5);
+		SampleState.OverscanFraction = CameraInfo.ViewInfo.GetOverscan();
+		SampleState.CropRectangle = InResolutionInfo.AccumulatorResolutionCropRect;
+		SampleState.bAllowOCIO = ParentNodeThisFrame->GetAllowOCIO();
+		SampleState.bForceLosslessCompression = ParentNodeThisFrame->GetForceLosslessCompression();
+		SampleState.bAllowsCompositing = ParentNodeThisFrame->GetAllowsCompositing();
+		SampleState.bIsBeautyPass = true;
+		SampleState.SceneCaptureSource = SceneCaptureSource;
+		SampleState.CompositingSortOrder = 10;
+		SampleState.RenderLayerIndex = LayerData.LayerIndex;
+	}
+
+	ApplyMovieGraphOverridesToSampleState(SampleState);
+
+	PostRendererSubmission(SampleState, InRenderTargetInitParams, Canvas, CameraInfo);
+}
+
+AccumulationDOF::FApertureSamplerConfig FMovieGraphAccumulationDOFPass::BuildSamplerConfig(UMovieGraphEvaluatedConfig* InEvaluatedConfig) const
+{
+	AccumulationDOF::FApertureSamplerConfig Config;
+
+	Config.World = GetRenderer()->GetWorld();
+
+	// Resolution & overscan: BackbufferResolution is the per-tile resolution with overscan
+	const FResolutionAndCameraInfo ResolutionAndCameraInfo = GetResolutionAndCameraInfo(InEvaluatedConfig);
+	Config.Resolution = ResolutionAndCameraInfo.BackbufferResolution;
+	Config.OverscanFraction = ResolutionAndCameraInfo.OverscanFraction;
+
+	// Sampling settings
+	Config.Mode = AccumulationDOF::EApertureSamplerMode::OneShot;
+
+	// Get settings from component (should always be valid if there's an Accumulation DOF pass active)
+	if (CachedDOFComponent.IsValid())
+	{
+		Config.NumSamples                          = CachedDOFComponent->NumSamples;
+		Config.SamplingPattern                     = CachedDOFComponent->SamplingPattern;
+		Config.DOFSplatSize	                       = CachedDOFComponent->DOFSplatSize;
+		Config.AxialChromaticAberrationIntensity   = CachedDOFComponent->AxialChromaticAberrationIntensity;
+		Config.AxialChromaticAberrationNumBands    = CachedDOFComponent->AxialChromaticAberrationNumBands;
+		Config.bSpectralLateralChromaticAberration = CachedDOFComponent->bSpectralLateralChromaticAberration;
+		Config.SphericalAberration                 = CachedDOFComponent->SphericalAberration;
+		Config.ComaAberration                      = CachedDOFComponent->ComaAberration;
+		Config.BokehTexture                        = CachedDOFComponent->BokehTexture;
+		Config.bEnableBokehTexture                 = CachedDOFComponent->bEnableBokehTexture;
+		Config.WeightChannel                       = CachedDOFComponent->WeightChannel;
+		Config.TintStrength                        = CachedDOFComponent->TintStrength;
+		Config.BokehEdgeSoftness                   = CachedDOFComponent->BokehEdgeSoftness;
+		Config.TemporalHistoryMode                 = CachedDOFComponent->TemporalHistoryMode;
+		Config.bUseJitterAA                        = CachedDOFComponent->bUseJitterAA;
+	}
+
+	Config.ViewModeIndex = VMI_Lit;
+
+	return Config;
+}
+
+AccumulationDOF::FApertureSamplerCameraState FMovieGraphAccumulationDOFPass::BuildCameraState() const
+{
+	AccumulationDOF::FApertureSamplerCameraState State;
+
+	State.CineCameraComponent = CachedCineCameraComponent;
+
+	if (CachedCineCameraComponent.IsValid())
+	{
+		if (const AActor* Owner = CachedCineCameraComponent->GetOwner())
+		{
+			State.CameraLocation = Owner->GetActorLocation();
+			State.CameraRotation = Owner->GetActorRotation();
+		}
+
+		State.bInitialized = true;
+	}
+
+	State.SceneFringeIntensity = CachedSceneFringeIntensity;
+
+	return State;
+}
+
+void FMovieGraphAccumulationDOFPass::InitializeSampler(UMovieGraphEvaluatedConfig* InEvaluatedConfig)
+{
+	// Find camera first
+	CachedCineCameraComponent = FindCineCameraComponent();
+	FindAccumulationDOFComponent();
+
+	// Build config and camera state
+	const AccumulationDOF::FApertureSamplerConfig Config = BuildSamplerConfig(InEvaluatedConfig);
+	const AccumulationDOF::FApertureSamplerCameraState CameraState = BuildCameraState();
+
+	// Create and initialize sampler
+	ApertureSampler = NewObject<UApertureSampler>(WeakGraphRenderer.Get());
+
+	if (!ApertureSampler->Initialize(Config, CameraState))
+	{
+		UE_LOG(LogAccumulationDOFMRG, Error, TEXT("Failed to initialize ApertureSampler"));
+
+		ApertureSampler->Shutdown();
+		ApertureSampler = nullptr;
+	}
+}
+
+UCineCameraComponent* FMovieGraphAccumulationDOFPass::FindCineCameraComponent() const
+{
+	const UE::MovieGraph::DefaultRenderer::FCameraInfo CameraInfo = GetRenderer()->GetCameraInfo(LayerData.CameraIndex);
+
+	const ACineCameraActor* CineCameraActor = Cast<ACineCameraActor>(CameraInfo.ViewActor);
+	if (!CineCameraActor)
+	{
+		return nullptr;
+	}
+
+	return CineCameraActor->GetCineCameraComponent();
+}
+
+UAccumulationDOFComponent* FMovieGraphAccumulationDOFPass::FindAccumulationDOFComponent()
+{
+	if (CachedDOFComponent.IsValid())
+	{
+		return CachedDOFComponent.Get();
+	}
+
+	const UE::MovieGraph::DefaultRenderer::FCameraInfo CameraInfo = GetRenderer()->GetCameraInfo(LayerData.CameraIndex);
+
+	const ACineCameraActor* CineCameraActor = Cast<ACineCameraActor>(CameraInfo.ViewActor);
+	if (!CineCameraActor)
+	{
+		UE_LOG(LogAccumulationDOFMRG, Warning, TEXT("FindAccumulationDOFComponent: No CineCameraActor available"));
+		return nullptr;
+	}
+
+	UAccumulationDOFComponent* DOFComponent = CineCameraActor->FindComponentByClass<UAccumulationDOFComponent>();
+	if (DOFComponent)
+	{
+		CachedDOFComponent = DOFComponent;
+	}
+
+	return DOFComponent;
+}
+
+float FMovieGraphAccumulationDOFPass::GetEffectiveDOFSplatSize() const
+{
+	if (CachedDOFComponent.IsValid())
+	{
+		return CachedDOFComponent->DOFSplatSize;
+	}
+
+	// Default
+	return 0.125f;
+}
diff --git a/Engine/Plugins/VirtualProduction/AccumulationDOF/Source/AccumulationDOF/Public/AccumulationDOF.h b/Engine/Plugins/VirtualProduction/AccumulationDOF/Source/AccumulationDOF/Public/AccumulationDOF.h
--- a/Engine/Plugins/VirtualProduction/AccumulationDOF/Source/AccumulationDOF/Public/AccumulationDOF.h
+++ b/Engine/Plugins/VirtualProduction/AccumulationDOF/Source/AccumulationDOF/Public/AccumulationDOF.h
@@ -12,4 +12,7 @@ public:
 	virtual void StartupModule() override;
 	virtual void ShutdownModule() override;
 
+private:
+	/** The name that will be given to the DOF pass that's provided to the pass factory in MRG's deferred renderer. */
+	static const FName MRGPassFactoryPassName;
 };
diff --git a/Engine/Plugins/VirtualProduction/AccumulationDOF/Source/AccumulationDOF/Public/AccumulationDOFComponent.h b/Engine/Plugins/VirtualProduction/AccumulationDOF/Source/AccumulationDOF/Public/AccumulationDOFComponent.h
--- a/Engine/Plugins/VirtualProduction/AccumulationDOF/Source/AccumulationDOF/Public/AccumulationDOFComponent.h
+++ b/Engine/Plugins/VirtualProduction/AccumulationDOF/Source/AccumulationDOF/Public/AccumulationDOFComponent.h
@@ -17,7 +17,7 @@ class UTexture2D;
  *
  * Attach to a CineCameraActor to configure DOF rendering settings for:
  * - Movie Render Queue (via AccumulationDOF pass)
- * - Editor viewport preview (Level Viewport >> Scalabilty & Performance >> Accumulation DOF)
+ * - Editor viewport preview (Level Viewport >> Scalability & Performance >> Accumulation DOF)
  */
 UCLASS(Blueprintable, ClassGroup = (Rendering), meta = (BlueprintSpawnableComponent, DisplayName = "Accumulation DOF"))
 class ACCUMULATIONDOF_API UAccumulationDOFComponent : public UActorComponent
@@ -25,6 +25,10 @@ class ACCUMULATIONDOF_API UAccumulationDOFComponent : public UActorComponent
 	GENERATED_BODY()
 
 public:
+	UAccumulationDOFComponent()
+	{
+		bAutoActivate = true;
+	}
 
 	/**
 	 * Specifies the number of aperture samples to render (use more for better quality, but will take more time).
@@ -42,6 +46,10 @@ public:
 		meta = (DisplayName = "DOF Splat Size", ClampMin = "0.0", ClampMax = "1.0", UIMin = "0.0", UIMax = "1.0"))
 	float DOFSplatSize = 0.125f;
 
+	/** Enable motion blur on the accumulated DOF result. */
+	UPROPERTY(EditAnywhere, BlueprintReadWrite, Category = "Accumulation Depth of Field")
+	bool bEnableMotionBlur = false;
+
 	/** Aperture sampling pattern. Hexaweb is generally preferred but Vogel Spirals are pretty good too. */
 	UPROPERTY(EditAnywhere, BlueprintReadWrite, Category = "Accumulation Depth of Field", AdvancedDisplay)
 	EApertureSamplingPattern SamplingPattern = EApertureSamplingPattern::Hexaweb;
@@ -119,5 +127,4 @@ public:
 	 */
 	UPROPERTY(EditAnywhere, BlueprintReadWrite, Category = "Monochromatic Aberrations", meta = (ClampMin = "0.0", UIMin = "0.0", UIMax = "1.0"))
 	float ComaAberration = 0.0f;
-
 };
diff --git a/Engine/Plugins/VirtualProduction/AccumulationDOF/Source/AccumulationDOF/Public/MovieGraphAccumulationDOFModifierNode.h b/Engine/Plugins/VirtualProduction/AccumulationDOF/Source/AccumulationDOF/Public/MovieGraphAccumulationDOFModifierNode.h
new file mode 100644
--- /dev/null
+++ b/Engine/Plugins/VirtualProduction/AccumulationDOF/Source/AccumulationDOF/Public/MovieGraphAccumulationDOFModifierNode.h
@@ -0,0 +1,105 @@
+// Copyright Epic Games, Inc. All Rights Reserved.
+
+#pragma once
+
+#include "Graph/MovieGraphRenderLayerSubsystem.h"
+#include "Graph/Nodes/MovieGraphModifierNode.h"
+
+#include "MovieGraphAccumulationDOFModifierNode.generated.h"
+
+#define UE_API ACCUMULATIONDOF_API
+
+/** Specifies how the enable state of the Accumulation DOF component will be set. */
+UENUM(BlueprintType)
+enum class EMovieGraphAccumulationDOFEnableType : uint8
+{
+	/** Use the current enable state on the camera (ie, no change is made). */
+	CameraActorDefault,
+
+	/** Use a custom value that will be used as the enable state. */
+	CustomValue
+};
+
+/** Specifies if the enable state of the Accumulation DOF component should be set, and if so, what the value of the enable state should be. */
+USTRUCT(BlueprintType)
+struct FMovieGraphAccumulationDOFEnableState
+{
+	GENERATED_BODY()
+
+	UE_API FMovieGraphAccumulationDOFEnableState() = default;
+
+	/** How the Accumulation DOF enable state should be treated. */
+	UPROPERTY(EditAnywhere, BlueprintReadWrite, Category = "Accumulation DOF")
+	EMovieGraphAccumulationDOFEnableType Type = EMovieGraphAccumulationDOFEnableType::CameraActorDefault;
+
+	/** If the enable Type is "Custom", this is the value the enable state should be set to. */
+	UPROPERTY(EditAnywhere, BlueprintReadWrite, Category = "Accumulation DOF", meta = (EditCondition = "Type == EMovieGraphAccumulationDOFEnableType::CustomValue"))
+	bool Value = false;
+};
+
+/** 
+ * Modifies the properties on the Accumulation DOF component (on all cameras within the world).
+ */
+UCLASS(MinimalAPI)
+class UMovieGraphAccumulationDOFModifierNode final : public UMovieGraphSettingNode, public IMovieGraphModifierNodeInterface
+{
+	GENERATED_BODY()
+
+public:
+	UE_API UMovieGraphAccumulationDOFModifierNode();
+
+#if WITH_EDITOR
+	//~ Begin UMovieGraphNode interface
+	UE_API virtual FText GetNodeTitle(const bool bGetDescriptive = false) const override;
+	UE_API virtual FText GetMenuCategory() const override;
+	UE_API virtual FLinearColor GetNodeTitleColor() const override;
+	UE_API virtual FSlateIcon GetIconAndTint(FLinearColor& OutColor) const override;
+	//~ End UMovieGraphNode interface
+#endif
+
+	//~ Begin IMovieGraphModifierNodeInterface interface
+	UE_API virtual TArray<UMovieGraphModifierBase*> GetAllModifiers() const override;
+	UE_API virtual bool SupportsCollections() const override;
+	//~ End IMovieGraphModifierNodeInterface interface
+
+	//~ Begin UMovieGraphSettingNode interface
+	UE_API virtual FString GetNodeInstanceName() const override;
+	//~ End UMovieGraphSettingNode interface
+
+public:
+	UPROPERTY(EditAnywhere, BlueprintReadWrite, Category = Overrides, meta = (InlineEditConditionToggle))
+	uint8 bOverride_EnableAccumulationDepthOfField : 1;
+
+	UPROPERTY(EditAnywhere, BlueprintReadWrite, Category = Overrides, meta = (InlineEditConditionToggle))
+	uint8 bOverride_NumSamples : 1;
+
+	UPROPERTY(EditAnywhere, BlueprintReadWrite, Category = Overrides, meta = (InlineEditConditionToggle))
+	uint8 bOverride_DOFSplatSize : 1;
+
+	/** 
+	 * Whether the Accumulation DOF component should be enabled/disabled.
+	 */
+	UPROPERTY(EditAnywhere, BlueprintReadWrite, Category = "Accumulation DOF", meta = (EditCondition = "bOverride_EnableAccumulationDepthOfField"))
+	FMovieGraphAccumulationDOFEnableState EnableAccumulationDepthOfField;
+
+	/** 
+	 * The number of samples that the Accumulation DOF component should use.
+	 */
+	UPROPERTY(EditAnywhere, BlueprintReadWrite, Category = "Accumulation DOF",
+		meta = (EditCondition = "bOverride_NumSamples", ClampMin = "1", UIMin = "1", UIMax = "32768"))
+	int32 NumSamples = 512;
+
+	/** 
+	 * The splat size that the Accumulation DOF component should use.
+	 */
+	UPROPERTY(EditAnywhere, BlueprintReadWrite, Category = "Accumulation DOF",
+		meta = (EditCondition = "bOverride_DOFSplatSize", DisplayName = "DOF Splat Size", ClampMin = "0.0", ClampMax = "1.0", UIMin = "0.0", UIMax = "1.0"))
+	float DOFSplatSize = 0.125f;
+
+private:
+	/** The modifier associated with this node. */
+	UPROPERTY(meta = (HideInActiveRenderSettings))
+	mutable TObjectPtr<class UMovieGraphAccumulationDOFModifier> Modifier;
+};
+
+#undef UE_API
diff --git a/Engine/Plugins/VirtualProduction/AccumulationDOF/Source/AccumulationDOF/Public/MovieGraphAccumulationDOFPass.h b/Engine/Plugins/VirtualProduction/AccumulationDOF/Source/AccumulationDOF/Public/MovieGraphAccumulationDOFPass.h
new file mode 100644
--- /dev/null
+++ b/Engine/Plugins/VirtualProduction/AccumulationDOF/Source/AccumulationDOF/Public/MovieGraphAccumulationDOFPass.h
@@ -0,0 +1,145 @@
+﻿// Copyright Epic Games, Inc. All Rights Reserved.
+
+#pragma once
+
+#include "CoreMinimal.h"
+#include "Graph/Renderers/MovieGraphDeferredPass.h"
+
+ACCUMULATIONDOF_API DECLARE_LOG_CATEGORY_EXTERN(LogAccumulationDOFMRG, Log, All);
+
+class ACineCameraActor;
+class FAccumulationDOFSceneViewExtension;
+class UAccumulationDOFComponent;
+class UApertureSampler;
+class UCineCameraComponent;
+class UTexture2D;
+class UTextureRenderTarget2D;
+
+namespace AccumulationDOF
+{
+	struct FApertureSamplerConfig;
+	struct FApertureSamplerCameraState;
+}
+
+/**
+ * Movie Graph deferred render pass for aperture-sampled depth of field.
+ *
+ * This pass renders multiple aperture samples with off-axis projection matrices and accumulates them 
+ * using GPU shaders to produce plausible DOF.
+ *
+ * Camera parameters (focal length, aperture, focus distance) are read from the CineCameraActor being 
+ * rendered by the sequence.
+ */
+struct ACCUMULATIONDOF_API FMovieGraphAccumulationDOFPass : public UE::MovieGraph::Rendering::FMovieGraphDeferredPass
+{
+	FMovieGraphAccumulationDOFPass();
+	virtual ~FMovieGraphAccumulationDOFPass() override;
+
+protected:
+	//~ FMovieGraphDeferredPass Interface
+	virtual void Setup(TWeakObjectPtr<UMovieGraphDefaultRenderer> InRenderer, TWeakObjectPtr<UMovieGraphImagePassBaseNode> InRenderPassNode, const FMovieGraphRenderPassLayerData& InLayer) override;
+	virtual void Teardown() override;
+	virtual void Render(const FMovieGraphTraversalContext& InFrameTraversalContext, const FMovieGraphTimeStepData& InTimeData) override;
+	virtual void AddReferencedObjects(FReferenceCollector& Collector) override;
+
+	//~ FMovieGraphImagePassBase Interface
+	virtual TSharedRef<FSceneViewFamilyContext> CreateSceneViewFamily(const UE::MovieGraph::Rendering::FViewFamilyInitData& InInitData) const override;
+	virtual FVector2D ComputeProjectionMatrixJitter(const FVector2D& InDefaultJitter, const FIntPoint& InBackbufferResolution) const override;
+	virtual FSceneViewInitOptions CreateViewInitOptions(const UE::MovieGraph::DefaultRenderer::FCameraInfo& InCameraInfo, FSceneViewFamilyContext* InViewFamily, FSceneViewStateInterface* InViewStateInterface) const override;
+	virtual void CalculateProjectionMatrix(UE::MovieGraph::DefaultRenderer::FCameraInfo& InOutCameraInfo, FSceneViewProjectionData& InOutProjectionData, const FIntPoint InBackbufferResolution, const FIntPoint InAccumulatorResolution) const override;
+	virtual void ApplyMovieGraphOverridesToSceneView(TSharedRef<FSceneViewFamilyContext> InOutFamily, const UE::MovieGraph::Rendering::FViewFamilyInitData& InInitData, const UE::MovieGraph::DefaultRenderer::FCameraInfo& InCameraInfo) const override;
+	virtual FSceneViewStateInterface* GetSceneViewState(UMovieGraphImagePassBaseNode* ParentNodeThisFrame, int32_t TileX, int32_t TileY) override;
+
+private:
+	/** Aperture sampler that handles sample rendering and accumulation. GC-protected via AddReferencedObjects. */
+	TObjectPtr<UApertureSampler> ApertureSampler = nullptr;
+
+	/** Current sample being rendered */
+	mutable int32 CurrentApertureSampleIndex = -1;
+
+	/** Focus distance override for axial CA (negative means no override). */
+	mutable float CurrentFocusDistanceOverride = -1.0f;
+
+	/** RGB weight for current spectral band in axial chromatic aberration. */
+	FVector3f CurrentAxialCAWeight = FVector3f(1.0f);
+
+	/** Cached reference to CineCameraComponent for SVE settings. */
+	TWeakObjectPtr<UCineCameraComponent> CachedCineCameraComponent;
+
+	/** Cached reference to AccumulationDOFComponent (found in CineCameraActor). */
+	TWeakObjectPtr<UAccumulationDOFComponent> CachedDOFComponent;
+
+	/** Scene View Extension for capture passes (one per aperture sample). */
+	TSharedPtr<FAccumulationDOFSceneViewExtension, ESPMode::ThreadSafe> CaptureExtension;
+
+	/** Scene View Extension for injection pass. */
+	TSharedPtr<FAccumulationDOFSceneViewExtension, ESPMode::ThreadSafe> InjectionExtension;
+
+	/** Flag to indicate we're doing a capture pass (so AddViewExtensions knows to add the capture SVE). */
+	bool bInCapturePass = false;
+
+	/** Flag to indicate we're doing the injection pass (so AddViewExtensions knows to add the injection SVE). */
+	bool bInInjectionPass = false;
+
+	/** Track whether we've initialized camera-dependent state. */
+	bool bCameraParamsInitialized = false;
+
+	/** Cached SceneFringeIntensity from the first sample's SVE (blended with PP volumes). */
+	float CachedSceneFringeIntensity = 0.0f;
+
+	/** Separate view state for injection renders. */
+	FSceneViewStateReference InjectionViewState;
+
+	/** Cached motion blur fraction from current frame's sample state. */
+	mutable float CachedMotionBlurFraction = 0.0f;
+
+	/** Whether this is the first injection render. */
+	bool bFirstInjection = true;
+
+	/** Saved LUT size before override (0 if not overridden). */
+	int32 SavedLUTSize = 0;
+
+private:
+	/** Build sampler config from graph settings and component. */
+	AccumulationDOF::FApertureSamplerConfig BuildSamplerConfig(UMovieGraphEvaluatedConfig* InEvaluatedConfig) const;
+
+	/** Build camera state from current sequence camera. */
+	AccumulationDOF::FApertureSamplerCameraState BuildCameraState() const;
+
+	/** Initialize the aperture sampler. */
+	void InitializeSampler(UMovieGraphEvaluatedConfig* InEvaluatedConfig);
+
+	/** Update camera-dependent state (called each frame). */
+	void UpdateCameraDependentState();
+
+	/** Find CineCameraComponent from MRG's camera binding. */
+	UCineCameraComponent* FindCineCameraComponent() const;
+
+	/** Find AccumulationDOFComponent on the camera actor. */
+	UAccumulationDOFComponent* FindAccumulationDOFComponent();
+
+	/** Get the DOF splat size from the component, or the default if there's no valid component. */
+	float GetEffectiveDOFSplatSize() const;
+
+	/** Render all aperture samples in one temporal sample. */
+	void RenderSamples(
+		const FMovieGraphTraversalContext& InFrameTraversalContext,
+		const FMovieGraphTimeStepData& InTimeData,
+		const FResolutionAndCameraInfo& InResolutionInfo,
+		const UE::MovieGraph::DefaultRenderer::FRenderTargetInitParams& InRenderTargetInitParams);
+
+	/** Render a single aperture sample. */
+	void RenderSingleApertureSample(
+		int32 SampleIndex,
+		const FMovieGraphTraversalContext& InFrameTraversalContext,
+		const FMovieGraphTimeStepData& InTimeData,
+		const FResolutionAndCameraInfo& InResolutionInfo,
+		const UE::MovieGraph::DefaultRenderer::FRenderTargetInitParams& InRenderTargetInitParams);
+
+	/** Injection pass and submit to MRG. */
+	void RenderInjectionAndSubmit(
+		const FMovieGraphTraversalContext& InFrameTraversalContext,
+		const FMovieGraphTimeStepData& InTimeData,
+		const FResolutionAndCameraInfo& InResolutionInfo,
+		const UE::MovieGraph::DefaultRenderer::FRenderTargetInitParams& InRenderTargetInitParams);
+};
