MoviePipeline: Fixed an issue where the Accumulation DOF modifier in the graph would not properly enable ADOF if an ADOF component did not already exist on the rendered camera.

#Synced-CL 52226548

diff --git a/Engine/Plugins/VirtualProduction/AccumulationDOF/Source/AccumulationDOF/Private/AccumulationDOF.cpp b/Engine/Plugins/VirtualProduction/AccumulationDOF/Source/AccumulationDOF/Private/AccumulationDOF.cpp
--- a/Engine/Plugins/VirtualProduction/AccumulationDOF/Source/AccumulationDOF/Private/AccumulationDOF.cpp
+++ b/Engine/Plugins/VirtualProduction/AccumulationDOF/Source/AccumulationDOF/Private/AccumulationDOF.cpp
@@ -60,17 +60,14 @@ void FAccumulationDOFModule::StartupModule()
 			return nullptr;
 		}
 
+		// Use the component's current active state if available, otherwise default to disabled
 		const UAccumulationDOFComponent* DOFComponent = CineCamera->FindComponentByClass<UAccumulationDOFComponent>();
-		if (!DOFComponent)
-		{
-			return nullptr;
-		}
-
-		// First, use the component's current active state to determine if this pass will be instantiated
-		bool bDOFEnabled = DOFComponent->IsActive();
+		bool bDOFEnabled = DOFComponent && DOFComponent->IsActive();
 
-		// Second, if an Accumulation DOF modifier on this layer is disabling Accumulation DOF, don't instantiate this pass.
-		// Checking the modifier is not ideal. However, the modifier is applied *after* pass instantiation, so we cannot solely
+		// Check if an Accumulation DOF modifier on this layer is overriding the enable state.
+		// This allows enabling Accumulation DOF even if no component exists on the camera,
+		// or disabling it even if the component is active. Checking the modifier is not ideal.
+		// However, the modifier is applied *after* pass instantiation, so we cannot solely
 		// rely on the component active state.
 		if (InContext.EvaluatedConfig)
 		{
