MoviePipeline: Fixed a rare bug in the MRG implementation of Accumulation DOF, which could cause random frames to appear distorted/incorrect.

#Synced-CL 52423999

diff --git a/Engine/Plugins/VirtualProduction/AccumulationDOF/Source/AccumulationDOF/Private/MovieGraphAccumulationDOFPass.cpp b/Engine/Plugins/VirtualProduction/AccumulationDOF/Source/AccumulationDOF/Private/MovieGraphAccumulationDOFPass.cpp
--- a/Engine/Plugins/VirtualProduction/AccumulationDOF/Source/AccumulationDOF/Private/MovieGraphAccumulationDOFPass.cpp
+++ b/Engine/Plugins/VirtualProduction/AccumulationDOF/Source/AccumulationDOF/Private/MovieGraphAccumulationDOFPass.cpp
@@ -758,6 +758,12 @@ void FMovieGraphAccumulationDOFPass::RenderInjectionAndSubmit(
 	ApplyMovieGraphOverridesToSampleState(SampleState);
 
 	PostRendererSubmission(SampleState, InRenderTargetInitParams, Canvas, CameraInfo);
+
+	// Flush render commands so the surface reader captures the injection result before the next frame's aperture
+	// samples reuse the pooled view RT. This serializes game/render threads briefly, but is negligible compared
+	// to the cost of rendering N aperture samples for this frame. The alternative (a dedicated staging RT that
+	// the injection result would be copied to immediately) would avoid the stall at the cost of extra VRAM.
+	FlushRenderingCommands();
 }
 
 AccumulationDOF::FApertureSamplerConfig FMovieGraphAccumulationDOFPass::BuildSamplerConfig(UMovieGraphEvaluatedConfig* InEvaluatedConfig) const
